Added an error message when the ConcertSyncServer module is unavailable after the multi-user server initialization

diff --git a/Engine/Plugins/Developer/Concert/ConcertSync/ConcertSyncServer/Source/ConcertSyncServer/Public/ConcertSyncServerLoop.inl b/Engine/Plugins/Developer/Concert/ConcertSync/ConcertSyncServer/Source/ConcertSyncServer/Public/ConcertSyncServerLoop.inl
--- a/Engine/Plugins/Developer/Concert/ConcertSync/ConcertSyncServer/Source/ConcertSyncServer/Public/ConcertSyncServerLoop.inl
+++ b/Engine/Plugins/Developer/Concert/ConcertSync/ConcertSyncServer/Source/ConcertSyncServer/Public/ConcertSyncServerLoop.inl
@@ -49,6 +49,12 @@ int32 ConcertSyncServerLoop(int32 ArgC, TCHAR** ArgV, const FConcertSyncServerLo
 	// Double CTRL+C signal will also cause process to terminate.
 	FPlatformMisc::SetGracefulTerminationHandler();
 
+	if (!IConcertSyncServerModule::IsAvailable())
+	{
+		UE_LOG(LogSyncServer, Error, TEXT("ConcertSyncServer Module is needed for proper execution. Initialization failed!"));
+		return -1;
+	}
+
 	// Get the server config.
 	const UConcertServerConfig* ServerConfig = nullptr;
 	if (InitArgs.GetServerConfigFunc)
