Fixed locking of newly added assets and saved assets in Multi-User

diff --git a/Engine/Plugins/Developer/Concert/ConcertSync/ConcertSyncServer/Source/ConcertSyncServer/Private/ConcertServerWorkspace.cpp b/Engine/Plugins/Developer/Concert/ConcertSync/ConcertSyncServer/Source/ConcertSyncServer/Private/ConcertServerWorkspace.cpp
--- a/Engine/Plugins/Developer/Concert/ConcertSync/ConcertSyncServer/Source/ConcertSyncServer/Private/ConcertServerWorkspace.cpp
+++ b/Engine/Plugins/Developer/Concert/ConcertSync/ConcertSyncServer/Source/ConcertSyncServer/Private/ConcertServerWorkspace.cpp
@@ -582,7 +582,7 @@ bool FConcertServerWorkspace::UnlockWorkspaceResource(const FName InResourceName
 	{
 		if (Owner->EndpointId == InLockEndpointId || bForce)
 		{
-			if (Owner->bExplicit == bExplicit || bForce)
+			if (!Owner->bExplicit || bExplicit || bForce)
 			{
 				const bool bWasTemporary = Owner->bTemporary;
 				LockedResources->Remove(InResourceName);
