Properly go through shutdown process and flush logs when shutting down the server abnormaly

diff --git a/Engine/Plugins/Developer/Concert/ConcertSync/ConcertSyncServer/Source/ConcertSyncServer/Public/ConcertSyncServerLoop.inl b/Engine/Plugins/Developer/Concert/ConcertSync/ConcertSyncServer/Source/ConcertSyncServer/Public/ConcertSyncServerLoop.inl
--- a/Engine/Plugins/Developer/Concert/ConcertSync/ConcertSyncServer/Source/ConcertSyncServer/Public/ConcertSyncServerLoop.inl
+++ b/Engine/Plugins/Developer/Concert/ConcertSync/ConcertSyncServer/Source/ConcertSyncServer/Public/ConcertSyncServerLoop.inl
@@ -27,11 +27,7 @@ int32 ConcertSyncServerLoop(int32 ArgC, TCHAR** ArgV, const FConcertSyncServerLo
 
 	// start up the main loop, adding some extra command line arguments:
 	//	-Messaging enables MessageBus transports
-	//	-stdout prevents double display logging
 	int32 Result = GEngineLoop.PreInit(ArgC, ArgV, TEXT(" -Messaging"));
-	//TODO: handle Result error? 
-
-	// TODO: need config? trim engine loop?
 	check(GConfig && GConfig->IsReadyForUse());
 
 	GLogConsole->Show(true);
@@ -52,107 +48,114 @@ int32 ConcertSyncServerLoop(int32 ArgC, TCHAR** ArgV, const FConcertSyncServerLo
 	if (!IConcertSyncServerModule::IsAvailable())
 	{
 		UE_LOG(LogSyncServer, Error, TEXT("ConcertSyncServer Module is needed for proper execution. Initialization failed!"));
-		return -1;
+		Result = -1;
 	}
 
-	// Get the server config.
 	const UConcertServerConfig* ServerConfig = nullptr;
-	if (InitArgs.GetServerConfigFunc)
-	{
-		ServerConfig = InitArgs.GetServerConfigFunc();
-	}
-	else
-	{
-		ServerConfig = IConcertSyncServerModule::Get().ParseServerSettings(FCommandLine::Get());
-	}
-	if (!ServerConfig)
-	{
-		UE_LOG(LogSyncServer, Error, TEXT("%s Configuration Failed!"), *InitArgs.ServiceFriendlyName);
-		return -1;
-	}
 
-	TSharedPtr<IPlugin> Plugin = IPluginManager::Get().FindPlugin("UdpMessaging");
-	if (!Plugin || !Plugin->IsEnabled())
+	if (Result >= 0)
 	{
-		// The UdpMessaging plugin should be added to the {appname}.Target.cs build file.
-		UE_LOG(LogSyncServer, Warning, TEXT("The 'UDP Messaging' plugin is disabled. The Concert server only supports UDP protocol."));
+		// Get the server config.
+		if (InitArgs.GetServerConfigFunc)
+		{
+			ServerConfig = InitArgs.GetServerConfigFunc();
+		}
+		else
+		{
+			ServerConfig = IConcertSyncServerModule::Get().ParseServerSettings(FCommandLine::Get());
+		}
+		if (!ServerConfig)
+		{
+			UE_LOG(LogSyncServer, Error, TEXT("%s Configuration Failed!"), *InitArgs.ServiceFriendlyName);
+			Result = -1;
+		}
 	}
 
-	// Setup Concert Sync to run in server mode.
-	TSharedPtr<IConcertSyncServer> ConcertSyncServer = IConcertSyncServerModule::Get().CreateServer(InitArgs.ServiceRole);
-	ConcertSyncServer->Startup(ServerConfig, InitArgs.SessionFlags);
-
-	// if we have a default session, set it up properly
-	if (!ServerConfig->DefaultSessionName.IsEmpty())
+	if (Result >= 0)
 	{
-		IConcertServerRef ConcertServer = ConcertSyncServer->GetConcertServer();
-		if (!ConcertServer->GetLiveSessionIdByName(ServerConfig->DefaultSessionName).IsValid())
+		TSharedPtr<IPlugin> Plugin = IPluginManager::Get().FindPlugin("UdpMessaging");
+		if (!Plugin || !Plugin->IsEnabled())
 		{
-			FConcertSessionInfo SessionInfo = ConcertServer->CreateSessionInfo();
-			SessionInfo.SessionName = ServerConfig->DefaultSessionName;
-			SessionInfo.Settings = ServerConfig->DefaultSessionSettings;
+			// The UdpMessaging plugin should be added to the {appname}.Target.cs build file.
+			UE_LOG(LogSyncServer, Warning, TEXT("The 'UDP Messaging' plugin is disabled. The Concert server only supports UDP protocol."));
+		}
 
-			bool bSuccess = false;
-			if (!ServerConfig->DefaultSessionToRestore.IsEmpty())
+		// Setup Concert Sync to run in server mode.
+		TSharedPtr<IConcertSyncServer> ConcertSyncServer = IConcertSyncServerModule::Get().CreateServer(InitArgs.ServiceRole);
+		ConcertSyncServer->Startup(ServerConfig, InitArgs.SessionFlags);
+
+		// if we have a default session, set it up properly
+		if (!ServerConfig->DefaultSessionName.IsEmpty())
+		{
+			IConcertServerRef ConcertServer = ConcertSyncServer->GetConcertServer();
+			if (!ConcertServer->GetLiveSessionIdByName(ServerConfig->DefaultSessionName).IsValid())
 			{
-				const FGuid ArchivedSessionId = ConcertServer->GetArchivedSessionIdByName(ServerConfig->DefaultSessionToRestore);
-				if (ArchivedSessionId.IsValid())
+				FConcertSessionInfo SessionInfo = ConcertServer->CreateSessionInfo();
+				SessionInfo.SessionName = ServerConfig->DefaultSessionName;
+				SessionInfo.Settings = ServerConfig->DefaultSessionSettings;
+
+				bool bSuccess = false;
+				if (!ServerConfig->DefaultSessionToRestore.IsEmpty())
 				{
-					FText FailureReason;
-					bSuccess = ConcertServer->RestoreSession(ArchivedSessionId, SessionInfo, FConcertSessionFilter(), FailureReason).IsValid();
-					if (!bSuccess)
+					const FGuid ArchivedSessionId = ConcertServer->GetArchivedSessionIdByName(ServerConfig->DefaultSessionToRestore);
+					if (ArchivedSessionId.IsValid())
 					{
-						UE_LOG(LogSyncServer, Error, TEXT("%s failed to restore archived session '%s' due to '%s'. Creating new session instead..."), *InitArgs.ServiceFriendlyName, *ServerConfig->DefaultSessionToRestore, *FailureReason.ToString());
+						FText FailureReason;
+						bSuccess = ConcertServer->RestoreSession(ArchivedSessionId, SessionInfo, FConcertSessionFilter(), FailureReason).IsValid();
+						if (!bSuccess)
+						{
+							UE_LOG(LogSyncServer, Error, TEXT("%s failed to restore archived session '%s' due to '%s'. Creating new session instead..."), *InitArgs.ServiceFriendlyName, *ServerConfig->DefaultSessionToRestore, *FailureReason.ToString());
+						}
+					}
+					else
+					{
+						UE_LOG(LogSyncServer, Error, TEXT("%s could not find archived session '%s' to restore. Creating new session instead..."), *InitArgs.ServiceFriendlyName, *ServerConfig->DefaultSessionToRestore);
 					}
 				}
-				else
-				{
-					UE_LOG(LogSyncServer, Error, TEXT("%s could not find archived session '%s' to restore. Creating new session instead..."), *InitArgs.ServiceFriendlyName, *ServerConfig->DefaultSessionToRestore);
-				}
-			}
 
-			if (!bSuccess)
-			{
-				FText FailureReason;
-				bSuccess = ConcertServer->CreateSession(SessionInfo, FailureReason).IsValid();
 				if (!bSuccess)
 				{
-					UE_LOG(LogSyncServer, Error, TEXT("%s failed to create session '%s' due to '%s'."), *InitArgs.ServiceFriendlyName, *ServerConfig->DefaultSessionName, *FailureReason.ToString());
+					FText FailureReason;
+					bSuccess = ConcertServer->CreateSession(SessionInfo, FailureReason).IsValid();
+					if (!bSuccess)
+					{
+						UE_LOG(LogSyncServer, Error, TEXT("%s failed to create session '%s' due to '%s'."), *InitArgs.ServiceFriendlyName, *ServerConfig->DefaultSessionName, *FailureReason.ToString());
+					}
 				}
 			}
 		}
-	}
 
-	UE_LOG(LogSyncServer, Display, TEXT("%s Initialized (Name: %s, Version: %d.%d, Role: %s)"), *InitArgs.ServiceFriendlyName, *ConcertSyncServer->GetConcertServer()->GetServerInfo().ServerName, ENGINE_MAJOR_VERSION, ENGINE_MINOR_VERSION, *ConcertSyncServer->GetConcertServer()->GetRole());
+		UE_LOG(LogSyncServer, Display, TEXT("%s Initialized (Name: %s, Version: %d.%d, Role: %s)"), *InitArgs.ServiceFriendlyName, *ConcertSyncServer->GetConcertServer()->GetServerInfo().ServerName, ENGINE_MAJOR_VERSION, ENGINE_MINOR_VERSION, *ConcertSyncServer->GetConcertServer()->GetRole());
 
-	double LastTime = FPlatformTime::Seconds();
-	const float IdealFrameTime = 1.0f / InitArgs.IdealFramerate;
+		double LastTime = FPlatformTime::Seconds();
+		const float IdealFrameTime = 1.0f / InitArgs.IdealFramerate;
 
-	while (!GIsRequestingExit)
-	{
-		const double CurrentTime = FPlatformTime::Seconds();
-		const double DeltaTime = CurrentTime - LastTime;
+		while (!GIsRequestingExit)
+		{
+			const double CurrentTime = FPlatformTime::Seconds();
+			const double DeltaTime = CurrentTime - LastTime;
 
-		FTaskGraphInterface::Get().ProcessThreadUntilIdle(ENamedThreads::GameThread);
+			FTaskGraphInterface::Get().ProcessThreadUntilIdle(ENamedThreads::GameThread);
 
-		// Pump & Tick objects
-		FTicker::GetCoreTicker().Tick(DeltaTime);
+			// Pump & Tick objects
+			FTicker::GetCoreTicker().Tick(DeltaTime);
 
-		GFrameCounter++;
-		FStats::AdvanceFrame(false);
-		GLog->FlushThreadedLogs();
+			GFrameCounter++;
+			FStats::AdvanceFrame(false);
+			GLog->FlushThreadedLogs();
 
-		// Run garbage collection for the UObjects for the rest of the frame or at least to 2 ms
-		IncrementalPurgeGarbage(true, FMath::Max<float>(0.002f, IdealFrameTime - (FPlatformTime::Seconds() - LastTime)));
+			// Run garbage collection for the UObjects for the rest of the frame or at least to 2 ms
+			IncrementalPurgeGarbage(true, FMath::Max<float>(0.002f, IdealFrameTime - (FPlatformTime::Seconds() - LastTime)));
 
-		// Throttle main thread main fps by sleeping if we still have time
-		FPlatformProcess::Sleep(FMath::Max<float>(0.0f, IdealFrameTime - (FPlatformTime::Seconds() - LastTime)));
-		
-		LastTime = CurrentTime;
-	}
+			// Throttle main thread main fps by sleeping if we still have time
+			FPlatformProcess::Sleep(FMath::Max<float>(0.0f, IdealFrameTime - (FPlatformTime::Seconds() - LastTime)));
+
+			LastTime = CurrentTime;
+		}
 
-	ConcertSyncServer->Shutdown();
-	ConcertSyncServer.Reset();
+		ConcertSyncServer->Shutdown();
+		ConcertSyncServer.Reset();
+	}
 
 	UE_LOG(LogSyncServer, Display, TEXT("%s Shutdown"), *InitArgs.ServiceFriendlyName);
 
diff --git a/Engine/Source/Programs/UnrealMultiUserServer/Private/UnrealMultiUserServer.cpp b/Engine/Source/Programs/UnrealMultiUserServer/Private/UnrealMultiUserServer.cpp
--- a/Engine/Source/Programs/UnrealMultiUserServer/Private/UnrealMultiUserServer.cpp
+++ b/Engine/Source/Programs/UnrealMultiUserServer/Private/UnrealMultiUserServer.cpp
@@ -20,7 +20,5 @@ INT32_MAIN_INT32_ARGC_TCHAR_ARGV()
 	ServerLoopInitArgs.ServiceRole = TEXT("MultiUser");
 	ServerLoopInitArgs.ServiceFriendlyName = TEXT("Multi-User Editing Server");
 
-	int32 Result = ConcertSyncServerLoop(ArgC, ArgV, ServerLoopInitArgs);
-	GLog->Flush();
-	return Result;
+	return ConcertSyncServerLoop(ArgC, ArgV, ServerLoopInitArgs);
 }
