Added the possibility to annotate concert events with message bus annotations
Added a concert custom message flags to indicate that a concert messages should be uniquely identifiable acoss clients and uses annotations to do so.
Send snapshot transactions, presence and sequencer events are uniquely identifiable.
Merged from CL 7437409

diff --git a/Engine/Plugins/Developer/Concert/ConcertMain/Source/Concert/Private/ConcertClientSession.cpp b/Engine/Plugins/Developer/Concert/ConcertMain/Source/Concert/Private/ConcertClientSession.cpp
--- a/Engine/Plugins/Developer/Concert/ConcertMain/Source/Concert/Private/ConcertClientSession.cpp
+++ b/Engine/Plugins/Developer/Concert/ConcertMain/Source/Concert/Private/ConcertClientSession.cpp
@@ -9,7 +9,8 @@
 #include "Misc/Paths.h"
 #include "UObject/StructOnScope.h"
 
-// TODO: change ptr to ref
+const FName ConcertClientMessageIdName("ConcertMessageId");
+
 FConcertClientSession::FConcertClientSession(const FConcertSessionInfo& InSessionInfo, const FConcertClientInfo& InClientInfo, const FConcertClientSettings& InSettings, TSharedPtr<IConcertLocalEndpoint> InClientSessionEndpoint, const FString& InSessionDirectory)
 	: FConcertSessionCommonImpl(InSessionInfo)
 	, ClientInfo(InClientInfo)
@@ -174,8 +175,15 @@ void FConcertClientSession::InternalSendCustomEvent(const UScriptStruct* EventTy
 	FConcertSession_CustomEvent CustomEvent;
 	CommonBuildCustomEvent(EventType, EventData, GetSessionClientEndpointId(), DestinationEndpointIds, CustomEvent);
 
+	// If the message is sent with the UniqueId flag, add an annotations so that we can uniquely identify multiple message bus copies of that message.
+	TMap<FName, FString> Annotations;
+	if (EnumHasAnyFlags(Flags, EConcertMessageFlags::UniqueId))
+	{
+		Annotations.Add(ConcertClientMessageIdName, FGuid::NewGuid().ToString());
+	}
+
 	// Send the event
-	ClientSessionEndpoint->SendEvent(CustomEvent, SessionInfo.ServerEndpointId, Flags);
+	ClientSessionEndpoint->SendEvent(CustomEvent, SessionInfo.ServerEndpointId, Flags, Annotations);
 }
 
 void FConcertClientSession::InternalSendCustomRequest(const UScriptStruct* RequestType, const void* RequestData, const FGuid& DestinationEndpointId, const TSharedRef<IConcertSessionCustomResponseHandler>& Handler)
diff --git a/Engine/Plugins/Developer/Concert/ConcertMain/Source/Concert/Private/ConcertServerSession.cpp b/Engine/Plugins/Developer/Concert/ConcertMain/Source/Concert/Private/ConcertServerSession.cpp
--- a/Engine/Plugins/Developer/Concert/ConcertMain/Source/Concert/Private/ConcertServerSession.cpp
+++ b/Engine/Plugins/Developer/Concert/ConcertMain/Source/Concert/Private/ConcertServerSession.cpp
@@ -8,6 +8,8 @@
 #include "Containers/Ticker.h"
 #include "UObject/StructOnScope.h"
 
+const FName ConcertServerMessageIdName("ConcertMessageId");
+
 FConcertServerSession::FConcertServerSession(const FConcertSessionInfo& InSessionInfo, const FConcertServerSettings& InSettings, TSharedPtr<IConcertLocalEndpoint> InServerSessionEndpoint, const FString& InSessionDirectory)
 	: FConcertSessionCommonImpl(InSessionInfo)
 	, ServerSessionEndpoint(MoveTemp(InServerSessionEndpoint))
@@ -105,10 +107,15 @@ void FConcertServerSession::InternalSendCustomEvent(const UScriptStruct* EventTy
 	CommonBuildCustomEvent(EventType, EventData, SessionInfo.ServerEndpointId, DestinationEndpointIds, CustomEvent);
 
 	// TODO: Optimize this so we can queue the event for multiple client endpoints at the same time
+	TMap<FName, FString> Annotations;
+	if (EnumHasAnyFlags(Flags, EConcertMessageFlags::UniqueId))
+	{
+		Annotations.Add(ConcertServerMessageIdName, FGuid::NewGuid().ToString());
+	}
 	for (const FGuid& DestinationEndpointId : DestinationEndpointIds)
 	{
 		// Send the event
-		ServerSessionEndpoint->SendEvent(CustomEvent, DestinationEndpointId, Flags);
+		ServerSessionEndpoint->SendEvent(CustomEvent, DestinationEndpointId, Flags, Annotations);
 	}
 }
 
@@ -253,7 +260,7 @@ void FConcertServerSession::HandleCustomEvent(const FConcertMessageContext& Cont
 		else if (const FSessionClient* Client = SessionClients.Find(DestinationEndpointId))
 		{
 			// Forward onto the client
-			ServerSessionEndpoint->SendEvent(*Message, Client->ClientInfo.ClientEndpointId, Message->IsReliable() ? EConcertMessageFlags::ReliableOrdered : EConcertMessageFlags::None);
+			ServerSessionEndpoint->SendEvent(*Message, Client->ClientInfo.ClientEndpointId, Message->IsReliable() ? EConcertMessageFlags::ReliableOrdered : EConcertMessageFlags::None, Context.Annotations);
 		}
 	}
 }
diff --git a/Engine/Plugins/Developer/Concert/ConcertMain/Source/ConcertTransport/Private/ConcertLocalEndpoint.cpp b/Engine/Plugins/Developer/Concert/ConcertMain/Source/ConcertTransport/Private/ConcertLocalEndpoint.cpp
--- a/Engine/Plugins/Developer/Concert/ConcertMain/Source/ConcertTransport/Private/ConcertLocalEndpoint.cpp
+++ b/Engine/Plugins/Developer/Concert/ConcertMain/Source/ConcertTransport/Private/ConcertLocalEndpoint.cpp
@@ -213,7 +213,7 @@ void FConcertLocalEndpoint::InternalQueueResponse(const TSharedRef<IConcertRespo
 	}
 }
 
-void FConcertLocalEndpoint::InternalQueueEvent(const TSharedRef<IConcertEvent>& Event, const FGuid& Endpoint, EConcertMessageFlags Flags)
+void FConcertLocalEndpoint::InternalQueueEvent(const TSharedRef<IConcertEvent>& Event, const FGuid& Endpoint, EConcertMessageFlags Flags, const TMap<FName, FString>& Annotations)
 {
 	// Fill sending info
 	SetMessageSendingInfo(Event);
@@ -237,7 +237,7 @@ void FConcertLocalEndpoint::InternalQueueEvent(const TSharedRef<IConcertEvent>&
 		}
 	}
 
-	SendMessage(Event, RemoteEndpoint.ToSharedRef(), Event->GetCreationDate());
+	SendMessage(Event, RemoteEndpoint.ToSharedRef(), Event->GetCreationDate(), Annotations);
 }
 
 void FConcertLocalEndpoint::InternalPublishEvent(const TSharedRef<IConcertEvent>& Event)
@@ -427,7 +427,7 @@ void FConcertLocalEndpoint::PublishMessage(const TSharedRef<IConcertMessage>& Me
 	);
 }
 
-void FConcertLocalEndpoint::SendMessage(const TSharedRef<IConcertMessage>& Message, const FConcertRemoteEndpointRef& RemoteEndpoint, const FDateTime& UtcNow)
+void FConcertLocalEndpoint::SendMessage(const TSharedRef<IConcertMessage>& Message, const FConcertRemoteEndpointRef& RemoteEndpoint, const FDateTime& UtcNow, const TMap<FName, FString>& Annotations)
 {
 	if (!MessageEndpoint.IsValid())
 	{
@@ -443,6 +443,7 @@ void FConcertLocalEndpoint::SendMessage(const TSharedRef<IConcertMessage>& Messa
 		Message->ConstructMessage(), // Should be deleted by MessageBus
 		Message->GetMessageType(),
 		Message->IsReliable() ? EMessageFlags::Reliable : EMessageFlags::None,
+		Annotations, 
 		nullptr, // No Attachment
 		TArrayBuilder<FMessageAddress>().Add(RemoteEndpoint->GetAddress()),
 		FTimespan::Zero(), // No Delay
@@ -470,7 +471,7 @@ void FConcertLocalEndpoint::InternalHandleMessage(const TSharedRef<IMessageConte
 
 	// Setup Context
 	const FConcertMessageData* Message = (const FConcertMessageData*)Context->GetMessage();
-	const FConcertMessageContext ConcertContext(Message->ConcertEndpointId, UtcNow, Message, MessageTypeInfo);
+	const FConcertMessageContext ConcertContext(Message->ConcertEndpointId, UtcNow, Message, MessageTypeInfo, Context->GetAnnotations());
 	Logger.LogMessageReceived(ConcertContext, EndpointContext.EndpointId);
 
 	// Special endpoint discovery message handling, process discovery before passing down the message
diff --git a/Engine/Plugins/Developer/Concert/ConcertMain/Source/ConcertTransport/Private/ConcertLocalEndpoint.h b/Engine/Plugins/Developer/Concert/ConcertMain/Source/ConcertTransport/Private/ConcertLocalEndpoint.h
--- a/Engine/Plugins/Developer/Concert/ConcertMain/Source/ConcertTransport/Private/ConcertLocalEndpoint.h
+++ b/Engine/Plugins/Developer/Concert/ConcertMain/Source/ConcertTransport/Private/ConcertLocalEndpoint.h
@@ -44,7 +44,7 @@ protected:
 
 	virtual void InternalQueueRequest(const TSharedRef<IConcertRequest>& Request, const FGuid& Endpoint) override;
 	virtual void InternalQueueResponse(const TSharedRef<IConcertResponse>& Response, const FGuid& Endpoint) override;
-	virtual void InternalQueueEvent(const TSharedRef<IConcertEvent>& Event, const FGuid& Endpoint, EConcertMessageFlags Flags) override;
+	virtual void InternalQueueEvent(const TSharedRef<IConcertEvent>& Event, const FGuid& Endpoint, EConcertMessageFlags Flags, const TMap<FName, FString>& Annotations) override;
 	virtual void InternalPublishEvent(const TSharedRef<IConcertEvent>& Event) override;
 
 private:
@@ -76,7 +76,7 @@ private:
 	void PublishMessage(const TSharedRef<IConcertMessage>& Message);
 
 	/** Send a message to a specific remote endpoint */
-	void SendMessage(const TSharedRef<IConcertMessage>& Message, const FConcertRemoteEndpointRef& RemoteEndpoint, const FDateTime& UtcNow);
+	void SendMessage(const TSharedRef<IConcertMessage>& Message, const FConcertRemoteEndpointRef& RemoteEndpoint, const FDateTime& UtcNow, const TMap<FName, FString>& Annotations = TMap<FName, FString>());
 
 	/** Handle an incoming message from the message bus */
 	void InternalHandleMessage(const TSharedRef<IMessageContext, ESPMode::ThreadSafe>& Context);
diff --git a/Engine/Plugins/Developer/Concert/ConcertMain/Source/ConcertTransport/Private/ConcertMessageContext.cpp b/Engine/Plugins/Developer/Concert/ConcertMain/Source/ConcertTransport/Private/ConcertMessageContext.cpp
--- a/Engine/Plugins/Developer/Concert/ConcertMain/Source/ConcertTransport/Private/ConcertMessageContext.cpp
+++ b/Engine/Plugins/Developer/Concert/ConcertMain/Source/ConcertTransport/Private/ConcertMessageContext.cpp
@@ -7,6 +7,7 @@ FConcertMessageCapturedContext::FConcertMessageCapturedContext(const FConcertMes
 {
 	CapturedContext.SenderConcertEndpointId = InContext.SenderConcertEndpointId;
 	CapturedContext.UtcNow = InContext.UtcNow;
+	CapturedContext.Annotations = InContext.Annotations;
 	MallocCapturedContextMessage(InContext.Message, InContext.MessageType);
 }
 
@@ -20,6 +21,7 @@ FConcertMessageCapturedContext::FConcertMessageCapturedContext(FConcertMessageCa
 {
 	CapturedContext.SenderConcertEndpointId = Other.CapturedContext.SenderConcertEndpointId;
 	CapturedContext.UtcNow = Other.CapturedContext.UtcNow;
+	CapturedContext.Annotations = MoveTemp(Other.CapturedContext.Annotations);
 	Exchange(CapturedContext.MessageType, Other.CapturedContext.MessageType); // Swap null with Other
 	Exchange(CapturedContext.Message, Other.CapturedContext.Message); // Swap null with Other
 }
@@ -32,6 +34,7 @@ const FConcertMessageCapturedContext& FConcertMessageCapturedContext::operator=(
 
 		CapturedContext.SenderConcertEndpointId = Other.CapturedContext.SenderConcertEndpointId;
 		CapturedContext.UtcNow = Other.CapturedContext.UtcNow;
+		CapturedContext.Annotations = MoveTemp(Other.CapturedContext.Annotations);
 		Exchange(CapturedContext.MessageType, Other.CapturedContext.MessageType); // Swap null with Other
 		Exchange(CapturedContext.Message, Other.CapturedContext.Message); // Swap null with Other
 	}
diff --git a/Engine/Plugins/Developer/Concert/ConcertMain/Source/ConcertTransport/Public/ConcertMessageContext.h b/Engine/Plugins/Developer/Concert/ConcertMain/Source/ConcertTransport/Public/ConcertMessageContext.h
--- a/Engine/Plugins/Developer/Concert/ConcertMain/Source/ConcertTransport/Public/ConcertMessageContext.h
+++ b/Engine/Plugins/Developer/Concert/ConcertMain/Source/ConcertTransport/Public/ConcertMessageContext.h
@@ -17,11 +17,12 @@ struct FConcertMessageContext
 	}
 
 	/** Construct a context with a reference to the given message data */
-	FConcertMessageContext(const FGuid& InSenderConcertEndpointId, const FDateTime& InUtcNow, const FConcertMessageData* InMessage, const UScriptStruct* InMessageType)
+	FConcertMessageContext(const FGuid& InSenderConcertEndpointId, const FDateTime& InUtcNow, const FConcertMessageData* InMessage, const UScriptStruct* InMessageType, const TMap<FName, FString>& InAnnotations)
 		: SenderConcertEndpointId(InSenderConcertEndpointId)
 		, UtcNow(InUtcNow)
 		, Message(InMessage)
 		, MessageType(InMessageType)
+		, Annotations(InAnnotations)
 	{
 	}
 
@@ -37,6 +38,7 @@ struct FConcertMessageContext
 	FDateTime UtcNow;
 	const FConcertMessageData* Message;
 	const UScriptStruct* MessageType;
+	TMap<FName, FString> Annotations;
 };
 
 struct FConcertMessageCapturedContext
diff --git a/Engine/Plugins/Developer/Concert/ConcertMain/Source/ConcertTransport/Public/ConcertTransportMessages.h b/Engine/Plugins/Developer/Concert/ConcertMain/Source/ConcertTransport/Public/ConcertTransportMessages.h
--- a/Engine/Plugins/Developer/Concert/ConcertMain/Source/ConcertTransport/Public/ConcertTransportMessages.h
+++ b/Engine/Plugins/Developer/Concert/ConcertMain/Source/ConcertTransport/Public/ConcertTransportMessages.h
@@ -16,6 +16,8 @@ enum class EConcertMessageFlags : uint8
 	None = 0,
 	/** Guarantee that this message is received by the client(s) and processed in the order they were sent */
 	ReliableOrdered = 1<<0,
+	/** Message sent with this flag should be uniquely identifiable across clients. */
+	UniqueId = 1<<1, 
 };
 ENUM_CLASS_FLAGS(EConcertMessageFlags);
 
diff --git a/Engine/Plugins/Developer/Concert/ConcertMain/Source/ConcertTransport/Public/IConcertEndpoint.h b/Engine/Plugins/Developer/Concert/ConcertMain/Source/ConcertTransport/Public/IConcertEndpoint.h
--- a/Engine/Plugins/Developer/Concert/ConcertMain/Source/ConcertTransport/Public/IConcertEndpoint.h
+++ b/Engine/Plugins/Developer/Concert/ConcertMain/Source/ConcertTransport/Public/IConcertEndpoint.h
@@ -113,13 +113,13 @@ public:
 	 * @param Flags : flags for the event (i.e.: is the event reliable)
 	 */
 	template<typename EventType>
-	void SendEvent(const EventType& Event, const FGuid& Endpoint, EConcertMessageFlags Flags = EConcertMessageFlags::None) // TODO: Allow moving
+	void SendEvent(const EventType& Event, const FGuid& Endpoint, EConcertMessageFlags Flags = EConcertMessageFlags::None, const TMap<FName, FString>& Annotations = TMap<FName, FString>()) // TODO: Allow moving
 	{
 		static_assert(TIsDerivedFrom<EventType, FConcertEventData>::IsDerived, "Sent EventType need to be a UStruct deriving of FConcertEventData.");
 		typedef TConcertEvent<EventType> ConcertEventType;
 		TSharedRef<ConcertEventType> EventRef = MakeShared<ConcertEventType>(Event);
 
-		InternalQueueEvent(EventRef, Endpoint, Flags);
+		InternalQueueEvent(EventRef, Endpoint, Flags, Annotations);
 	}
 
 	/**
@@ -301,7 +301,7 @@ protected:
 	 * @param Endpoint : the remote endpoint Id to send to
 	 * Flags : flags for the event (i.e.: is the event reliable)	 
 	 */
-	virtual void InternalQueueEvent(const TSharedRef<IConcertEvent>& Event, const FGuid& Endpoint, EConcertMessageFlags Flags) = 0;
+	virtual void InternalQueueEvent(const TSharedRef<IConcertEvent>& Event, const FGuid& Endpoint, EConcertMessageFlags Flags, const TMap<FName, FString>& Annotations) = 0;
 
 	/**
 	 * Publish an event to any listening endpoints
diff --git a/Engine/Plugins/Developer/Concert/ConcertSync/ConcertSyncClient/Source/ConcertSyncClient/Private/ConcertClientPresenceMode.cpp b/Engine/Plugins/Developer/Concert/ConcertSync/ConcertSyncClient/Source/ConcertSyncClient/Private/ConcertClientPresenceMode.cpp
--- a/Engine/Plugins/Developer/Concert/ConcertSync/ConcertSyncClient/Source/ConcertSyncClient/Private/ConcertClientPresenceMode.cpp
+++ b/Engine/Plugins/Developer/Concert/ConcertSync/ConcertSyncClient/Source/ConcertSyncClient/Private/ConcertClientPresenceMode.cpp
@@ -87,7 +87,7 @@ void FConcertClientBasePresenceMode::SendEvents(IConcertClientSession& Session)
 
 	SetUpdateIndex(Session, FConcertClientPresenceDataUpdateEvent::StaticStruct()->GetFName(), PresenceDataUpdatedEvent);
 
-	Session.SendCustomEvent(PresenceDataUpdatedEvent, Session.GetSessionClientEndpointIds(), EConcertMessageFlags::None);
+	Session.SendCustomEvent(PresenceDataUpdatedEvent, Session.GetSessionClientEndpointIds(), EConcertMessageFlags::UniqueId);
 }
 
 FTransform FConcertClientBasePresenceMode::GetHeadTransform()
@@ -218,7 +218,7 @@ void FConcertClientDesktopPresenceMode::SendEvents(IConcertClientSession& Sessio
 
 				SetUpdateIndex(Session, FConcertClientDesktopPresenceUpdateEvent::StaticStruct()->GetFName(), Event);
 
-				Session.SendCustomEvent(Event, Session.GetSessionClientEndpointIds(), EConcertMessageFlags::None);
+				Session.SendCustomEvent(Event, Session.GetSessionClientEndpointIds(), EConcertMessageFlags::UniqueId);
 			}
 		}
 	}
@@ -302,7 +302,7 @@ void FConcertClientVRPresenceMode::SendEvents(IConcertClientSession& Session)
 
 		SetUpdateIndex(Session, FConcertClientDesktopPresenceUpdateEvent::StaticStruct()->GetFName(), Event);
 
-		Session.SendCustomEvent(Event, Session.GetSessionClientEndpointIds(), EConcertMessageFlags::None);
+		Session.SendCustomEvent(Event, Session.GetSessionClientEndpointIds(), EConcertMessageFlags::UniqueId);
 	}
 }
 
diff --git a/Engine/Plugins/Developer/Concert/ConcertSync/ConcertSyncServer/Source/ConcertSyncServer/Private/ConcertServerSequencerManager.cpp b/Engine/Plugins/Developer/Concert/ConcertSync/ConcertSyncServer/Source/ConcertSyncServer/Private/ConcertServerSequencerManager.cpp
--- a/Engine/Plugins/Developer/Concert/ConcertSync/ConcertSyncServer/Source/ConcertSyncServer/Private/ConcertServerSequencerManager.cpp
+++ b/Engine/Plugins/Developer/Concert/ConcertSync/ConcertSyncServer/Source/ConcertSyncServer/Private/ConcertServerSequencerManager.cpp
@@ -49,7 +49,7 @@ void FConcertServerSequencerManager::HandleSequencerStateEvent(const FConcertSes
 	// Forward the message to the other clients
 	TArray<FGuid> ClientIds = LiveSession->GetSession().GetSessionClientEndpointIds();
 	ClientIds.Remove(InEventContext.SourceEndpointId);
-	LiveSession->GetSession().SendCustomEvent(InEvent, ClientIds, EConcertMessageFlags::ReliableOrdered);
+	LiveSession->GetSession().SendCustomEvent(InEvent, ClientIds, EConcertMessageFlags::ReliableOrdered | EConcertMessageFlags::UniqueId);
 }
 
 void FConcertServerSequencerManager::HandleSequencerOpenEvent(const FConcertSessionContext& InEventContext, const FConcertSequencerOpenEvent& InEvent)
@@ -62,7 +62,7 @@ void FConcertServerSequencerManager::HandleSequencerOpenEvent(const FConcertSess
 	// Forward the message to the other clients
 	TArray<FGuid> ClientIds = LiveSession->GetSession().GetSessionClientEndpointIds();
 	ClientIds.Remove(InEventContext.SourceEndpointId);
-	LiveSession->GetSession().SendCustomEvent(InEvent, ClientIds, EConcertMessageFlags::ReliableOrdered);
+	LiveSession->GetSession().SendCustomEvent(InEvent, ClientIds, EConcertMessageFlags::ReliableOrdered | EConcertMessageFlags::UniqueId);
 }
 
 void FConcertServerSequencerManager::HandleSequencerCloseEvent(const FConcertSessionContext& InEventContext, const FConcertSequencerCloseEvent& InEvent)
@@ -77,13 +77,13 @@ void FConcertServerSequencerManager::HandleSequencerCloseEvent(const FConcertSes
 			FConcertSequencerCloseEvent CloseEvent;
 			CloseEvent.bMasterClose = false;
 			CloseEvent.SequenceObjectPath = InEvent.SequenceObjectPath;
-			LiveSession->GetSession().SendCustomEvent(CloseEvent, LiveSession->GetSession().GetSessionClientEndpointIds(), EConcertMessageFlags::ReliableOrdered);
+			LiveSession->GetSession().SendCustomEvent(CloseEvent, LiveSession->GetSession().GetSessionClientEndpointIds(), EConcertMessageFlags::ReliableOrdered | EConcertMessageFlags::UniqueId);
 			SequencerStates.Remove(*InEvent.SequenceObjectPath);
 		}
 		// if a sequence was close while it was the master, forward it to client
 		else if (InEvent.bMasterClose)
 		{
-			LiveSession->GetSession().SendCustomEvent(InEvent, LiveSession->GetSession().GetSessionClientEndpointIds(), EConcertMessageFlags::ReliableOrdered);
+			LiveSession->GetSession().SendCustomEvent(InEvent, LiveSession->GetSession().GetSessionClientEndpointIds(), EConcertMessageFlags::ReliableOrdered | EConcertMessageFlags::UniqueId);
 		}
 	}
 }
@@ -102,7 +102,7 @@ void FConcertServerSequencerManager::HandleSessionClientChanged(IConcertServerSe
 				// Forward the close event to clients
 				FConcertSequencerCloseEvent CloseEvent;
 				CloseEvent.SequenceObjectPath = It->Key.ToString();
-				LiveSession->GetSession().SendCustomEvent(CloseEvent, LiveSession->GetSession().GetSessionClientEndpointIds(), EConcertMessageFlags::ReliableOrdered);
+				LiveSession->GetSession().SendCustomEvent(CloseEvent, LiveSession->GetSession().GetSessionClientEndpointIds(), EConcertMessageFlags::ReliableOrdered|EConcertMessageFlags::UniqueId);
 
 				It.RemoveCurrent();
 			}
diff --git a/Engine/Plugins/Developer/Concert/ConcertSync/ConcertSyncServer/Source/ConcertSyncServer/Private/ConcertServerWorkspace.cpp b/Engine/Plugins/Developer/Concert/ConcertSync/ConcertSyncServer/Source/ConcertSyncServer/Private/ConcertServerWorkspace.cpp
--- a/Engine/Plugins/Developer/Concert/ConcertSync/ConcertSyncServer/Source/ConcertSyncServer/Private/ConcertServerWorkspace.cpp
+++ b/Engine/Plugins/Developer/Concert/ConcertSync/ConcertSyncServer/Source/ConcertSyncServer/Private/ConcertServerWorkspace.cpp
@@ -359,7 +359,7 @@ void FConcertServerWorkspace::HandleTransactionSnapshotEvent(const FConcertSessi
 		// If the client has the lock, then forward the snapshot
 		TArray<FGuid> QueueClientEndpointIds = LiveSession->GetSession().GetSessionClientEndpointIds();
 		QueueClientEndpointIds.Remove(InEventContext.SourceEndpointId);
-		LiveSession->GetSession().SendCustomEvent(InEvent, QueueClientEndpointIds, EConcertMessageFlags::None);
+		LiveSession->GetSession().SendCustomEvent(InEvent, QueueClientEndpointIds, EConcertMessageFlags::UniqueId);
 	}
 	// otherwise do nothing, we will reject the finalized transaction (TODO: send notification back?)
 }
