Added a missing line of code that was removed by error during merge.

diff --git a/Engine/Plugins/Developer/Concert/ConcertSync/ConcertSyncServer/Source/ConcertSyncServer/Private/ConcertServerWorkspace.cpp b/Engine/Plugins/Developer/Concert/ConcertSync/ConcertSyncServer/Source/ConcertSyncServer/Private/ConcertServerWorkspace.cpp
--- a/Engine/Plugins/Developer/Concert/ConcertSync/ConcertSyncServer/Source/ConcertSyncServer/Private/ConcertServerWorkspace.cpp
+++ b/Engine/Plugins/Developer/Concert/ConcertSync/ConcertSyncServer/Source/ConcertSyncServer/Private/ConcertServerWorkspace.cpp
@@ -56,7 +56,7 @@ void FConcertServerWorkspace::BindSession(const TSharedRef<FConcertSyncServerLiv
 	LiveSession->GetSession().RegisterCustomEventHandler<FConcertTransactionSnapshotEvent>(this, &FConcertServerWorkspace::HandleTransactionSnapshotEvent);
 
 	LiveSession->GetSession().RegisterCustomRequestHandler<FConcertResourceLockRequest, FConcertResourceLockResponse>(this, &FConcertServerWorkspace::HandleResourceLockRequest);
-
+	LiveSession->GetSession().RegisterCustomRequestHandler<FConcertSyncEventRequest, FConcertSyncEventResponse>(this, &FConcertServerWorkspace::HandleSyncEventRequest);
 	LiveSession->GetSession().RegisterCustomEventHandler<FConcertIgnoreActivityStateChangedEvent>(this, &FConcertServerWorkspace::HandleIgnoredActivityStateChanged);
 }
 
