- Added a option to configure whether the console should be shown, set it to false for Disaster Recovery.

diff --git a/Engine/Plugins/Developer/Concert/ConcertSync/ConcertSyncServer/Source/ConcertSyncServer/Public/ConcertSyncServerLoop.h b/Engine/Plugins/Developer/Concert/ConcertSync/ConcertSyncServer/Source/ConcertSyncServer/Public/ConcertSyncServerLoop.h
--- a/Engine/Plugins/Developer/Concert/ConcertSync/ConcertSyncServer/Source/ConcertSyncServer/Public/ConcertSyncServerLoop.h
+++ b/Engine/Plugins/Developer/Concert/ConcertSync/ConcertSyncServer/Source/ConcertSyncServer/Public/ConcertSyncServerLoop.h
@@ -27,6 +27,9 @@ struct FConcertSyncServerLoopInitArgs
 
 	/** Function to get the server settings object to configure the server with with, or unbound to parse the default settings */
 	TFunction<const UConcertServerConfig*()> GetServerConfigFunc;
+    
+    /** Whether the service should show the log console. */
+    bool bShowConsole = true;
 };
 
 /**
diff --git a/Engine/Plugins/Developer/Concert/ConcertSync/ConcertSyncServer/Source/ConcertSyncServer/Public/ConcertSyncServerLoop.inl b/Engine/Plugins/Developer/Concert/ConcertSync/ConcertSyncServer/Source/ConcertSyncServer/Public/ConcertSyncServerLoop.inl
--- a/Engine/Plugins/Developer/Concert/ConcertSync/ConcertSyncServer/Source/ConcertSyncServer/Public/ConcertSyncServerLoop.inl
+++ b/Engine/Plugins/Developer/Concert/ConcertSync/ConcertSyncServer/Source/ConcertSyncServer/Public/ConcertSyncServerLoop.inl
@@ -31,7 +31,10 @@ int32 ConcertSyncServerLoop(int32 ArgC, TCHAR** ArgV, const FConcertSyncServerLo
 	int32 Result = GEngineLoop.PreInit(ArgC, ArgV, TEXT(" -Messaging"));
 	check(GConfig && GConfig->IsReadyForUse());
 
-	GLogConsole->Show(true);
+	if (InitArgs.bShowConsole)
+	{
+		GLogConsole->Show(true);
+	}
 
 	FModuleManager::Get().StartProcessingNewlyLoadedObjects();
 
diff --git a/Engine/Source/Programs/UnrealDisasterRecoveryService/Private/UnrealDisasterRecoveryService.cpp b/Engine/Source/Programs/UnrealDisasterRecoveryService/Private/UnrealDisasterRecoveryService.cpp
--- a/Engine/Source/Programs/UnrealDisasterRecoveryService/Private/UnrealDisasterRecoveryService.cpp
+++ b/Engine/Source/Programs/UnrealDisasterRecoveryService/Private/UnrealDisasterRecoveryService.cpp
@@ -23,6 +23,7 @@ INT32_MAIN_INT32_ARGC_TCHAR_ARGV()
 	ServerLoopInitArgs.ServiceRole = TEXT("DisasterRecovery");
 	ServerLoopInitArgs.ServiceFriendlyName = TEXT("Disaster Recovery Service");
 	ServerLoopInitArgs.ServiceAutoArchiveSessionFilter.bIncludeIgnoredActivities = true;
+    ServerLoopInitArgs.bShowConsole = false;
 
 	ServerLoopInitArgs.GetServerConfigFunc = [&EditorProcessId]() -> const UConcertServerConfig*
 	{
