- Fixed concert server, when live sync is disabled (client doesn't track live transactions), to provide the last transaction event ID used to ignore all live transactions that occurred on the discarded package.

Details:

When some package modifications (in memory) are dimissed, the corresponding live transactions on that package must also be dismissed. In this case, a 'dummy' package event is placed in the DB with the last saved package data and the current 'transaction event ID' is stored. When the modifications are synchronized, the last saved package is pulled and hot swapped, then the 'live' transaction on the packages are requested. In this case, all transaction event ID <= the one recorded in the package dummy event are ignored (fence) when requested. The bug occurred because the client could not provide the transaction event ID because he doesn't track them (live sync is off) and therefore, it was always 0. Consequently, all transaction event ID > 0 (all transactions) were reapplied on top of the package.

diff --git a/Engine/Plugins/Developer/Concert/ConcertSync/ConcertSyncServer/Source/ConcertSyncServer/Private/ConcertServerWorkspace.cpp b/Engine/Plugins/Developer/Concert/ConcertSync/ConcertSyncServer/Source/ConcertSyncServer/Private/ConcertServerWorkspace.cpp
--- a/Engine/Plugins/Developer/Concert/ConcertSync/ConcertSyncServer/Source/ConcertSyncServer/Private/ConcertServerWorkspace.cpp
+++ b/Engine/Plugins/Developer/Concert/ConcertSync/ConcertSyncServer/Source/ConcertSyncServer/Private/ConcertServerWorkspace.cpp
@@ -296,6 +296,12 @@ void FConcertServerWorkspace::HandlePackageUpdateEvent(const FConcertSessionCont
 				{
 					PackageActivity.EventData.Package.Info.PackageFileExtension = HeadPackageInfo.PackageFileExtension;
 				}
+
+				// Without Live Sync, the clients don't track the ongoing transactions in their 'local' database and can't provide the current transaction event ID when unsaved modifications to a package are discarded (to discard corresponding live transactions on this package).
+				if (!EnumHasAnyFlags(LiveSession->GetSessionFlags(), EConcertSyncSessionFlags::EnableLiveSync))
+				{
+					LiveSession->GetSessionDatabase().GetTransactionMaxEventId(PackageActivity.EventData.Package.Info.TransactionEventIdAtSave); // Get the last transaction event from the server db.
+				}
 			}
 			PackageActivity.EventSummary.SetTypedPayload(FConcertSyncPackageActivitySummary::CreateSummaryForEvent(PackageActivity.EventData));
 			PackageActivity.bIgnored = ShouldIgnoreClientActivityOnRestore(Context.SourceEndpointId);
