Fixed PIE/SIE dummy package events not marked as 'ignored on restore'  by [REDACTED_PERSON] when PIE/SIE is launched from a Multi-User session.

diff --git a/Engine/Plugins/Developer/Concert/ConcertSync/ConcertSyncClient/Source/ConcertSyncClient/Private/ConcertClientPackageManager.cpp b/Engine/Plugins/Developer/Concert/ConcertSync/ConcertSyncClient/Source/ConcertSyncClient/Private/ConcertClientPackageManager.cpp
--- a/Engine/Plugins/Developer/Concert/ConcertSync/ConcertSyncClient/Source/ConcertSyncClient/Private/ConcertClientPackageManager.cpp
+++ b/Engine/Plugins/Developer/Concert/ConcertSync/ConcertSyncClient/Source/ConcertSyncClient/Private/ConcertClientPackageManager.cpp
@@ -177,7 +177,6 @@ void FConcertClientPackageManager::HandlePackageDiscarded(UPackage* InPackage)
 	Event.Package.Info.PackageName = InPackage->GetFName();
 	Event.Package.Info.PackageFileExtension = UWorld::FindWorldInPackage(InPackage) ? FPackageName::GetMapPackageExtension() : FPackageName::GetAssetPackageExtension();
 	Event.Package.Info.PackageUpdateType = EConcertPackageUpdateType::Dummy;
-	Event.bReplayable = bReplayableEvents;
 	LiveSession->GetSessionDatabase().GetTransactionMaxEventId(Event.Package.Info.TransactionEventIdAtSave);
 	LiveSession->GetSession().SendCustomEvent(Event, LiveSession->GetSession().GetSessionServerEndpointId(), EConcertMessageFlags::ReliableOrdered);
 }
@@ -330,7 +329,6 @@ void FConcertClientPackageManager::HandleLocalPackageEvent(const FConcertPackage
 
 	FConcertPackageUpdateEvent Event;
 	Event.Package = Package;
-	Event.bReplayable = bReplayableEvents;
 	LiveSession->GetSessionDatabase().GetTransactionMaxEventId(Event.Package.Info.TransactionEventIdAtSave);
 	LiveSession->GetSession().SendCustomEvent(Event, LiveSession->GetSession().GetSessionServerEndpointId(), EConcertMessageFlags::ReliableOrdered);
 }
diff --git a/Engine/Plugins/Developer/Concert/ConcertSync/ConcertSyncClient/Source/ConcertSyncClient/Private/ConcertClientPackageManager.h b/Engine/Plugins/Developer/Concert/ConcertSync/ConcertSyncClient/Source/ConcertSyncClient/Private/ConcertClientPackageManager.h
--- a/Engine/Plugins/Developer/Concert/ConcertSync/ConcertSyncClient/Source/ConcertSyncClient/Private/ConcertClientPackageManager.h
+++ b/Engine/Plugins/Developer/Concert/ConcertSync/ConcertSyncClient/Source/ConcertSyncClient/Private/ConcertClientPackageManager.h
@@ -77,11 +77,6 @@ public:
 	 */
 	bool PersistSessionChanges(TArrayView<const FString> InFilesToPersist, ISourceControlProvider* SourceControlProvider, TArray<FText>* OutFailureReasons = nullptr);
 
-	/**
-	 * Sets whether further package events emitted to the server has the 'replayable' flag on or off.
-	 */
-	void SetPackageEventAsReplayable(bool bReplayable) { bReplayableEvents = bReplayable; }
-
 private:
 	/**
 	 * Apply the data in the given package to disk and update the in-memory state.
@@ -168,9 +163,4 @@ private:
 	 * Array of package names that are pending an in-memory purge.
 	 */
 	TArray<FName> PackagesPendingPurge;
-
-	/**
-	 * Whether the package events emitted to the server has the replayable flag on or off.
-	 */
-	bool bReplayableEvents = false;
 };
diff --git a/Engine/Plugins/Developer/Concert/ConcertSync/ConcertSyncClient/Source/ConcertSyncClient/Private/ConcertClientTransactionManager.cpp b/Engine/Plugins/Developer/Concert/ConcertSync/ConcertSyncClient/Source/ConcertSyncClient/Private/ConcertClientTransactionManager.cpp
--- a/Engine/Plugins/Developer/Concert/ConcertSync/ConcertSyncClient/Source/ConcertSyncClient/Private/ConcertClientTransactionManager.cpp
+++ b/Engine/Plugins/Developer/Concert/ConcertSync/ConcertSyncClient/Source/ConcertSyncClient/Private/ConcertClientTransactionManager.cpp
@@ -391,12 +391,11 @@ void FConcertClientTransactionManager::HandleLocalTransactionFinalized(const FCo
 	}
 
 	FPendingTransactionToSend& PendingTransaction = HandleLocalTransactionCommon(InCommonData);
-	PendingTransaction.CommonData.bReplayable = bReplayableEvents;
 	PendingTransaction.FinalizedData = InFinalizedData;
 	PendingTransaction.bIsFinalized = true;
 }
 
-void FConcertClientTransactionManager::SendTransactionFinalizedEvent(const FGuid& InTransactionId, const FGuid& InOperationId, UObject* InPrimaryObject, const TArray<FName>& InModifiedPackages, const TArray<FConcertExportedObject>& InObjectUpdates, const FConcertLocalIdentifierTable& InLocalIdentifierTable, const FText& InTitle, bool bReplayable)
+void FConcertClientTransactionManager::SendTransactionFinalizedEvent(const FGuid& InTransactionId, const FGuid& InOperationId, UObject* InPrimaryObject, const TArray<FName>& InModifiedPackages, const TArray<FConcertExportedObject>& InObjectUpdates, const FConcertLocalIdentifierTable& InLocalIdentifierTable, const FText& InTitle)
 {
 	FConcertTransactionFinalizedEvent TransactionFinalizedEvent;
 	FillTransactionEvent(InTransactionId, InOperationId, InModifiedPackages, TransactionFinalizedEvent);
@@ -404,7 +403,6 @@ void FConcertClientTransactionManager::SendTransactionFinalizedEvent(const FGuid
 	TransactionFinalizedEvent.ExportedObjects = InObjectUpdates;
 	InLocalIdentifierTable.GetState(TransactionFinalizedEvent.LocalIdentifierState);
 	TransactionFinalizedEvent.Title = InTitle;
-	TransactionFinalizedEvent.bReplayable = bReplayable;
 
 	LiveSession->GetSession().SendCustomEvent(TransactionFinalizedEvent, LiveSession->GetSession().GetSessionServerEndpointId(), EConcertMessageFlags::ReliableOrdered);
 }
@@ -453,8 +451,7 @@ void FConcertClientTransactionManager::SendPendingTransactionEvents()
 						PendingTransactionPtr->CommonData.ModifiedPackages, 
 						PendingTransactionPtr->FinalizedData.FinalizedObjectUpdates, 
 						PendingTransactionPtr->FinalizedData.FinalizedLocalIdentifierTable, 
-						PendingTransactionPtr->CommonData.TransactionTitle,
-						PendingTransactionPtr->CommonData.bReplayable
+						PendingTransactionPtr->CommonData.TransactionTitle
 						);
 				}
 				// TODO: Warn about excluded objects?
diff --git a/Engine/Plugins/Developer/Concert/ConcertSync/ConcertSyncClient/Source/ConcertSyncClient/Private/ConcertClientTransactionManager.h b/Engine/Plugins/Developer/Concert/ConcertSync/ConcertSyncClient/Source/ConcertSyncClient/Private/ConcertClientTransactionManager.h
--- a/Engine/Plugins/Developer/Concert/ConcertSync/ConcertSyncClient/Source/ConcertSyncClient/Private/ConcertClientTransactionManager.h
+++ b/Engine/Plugins/Developer/Concert/ConcertSync/ConcertSyncClient/Source/ConcertSyncClient/Private/ConcertClientTransactionManager.h
@@ -52,11 +52,6 @@ public:
 	 */
 	void ProcessPending();
 
-	/**
-	 * Sets whether further events emitted to the server has the 'replayable' flag on or off.
-	 */
-	void SetTransactionEventsAsReplayable(bool bReplayable) { bReplayableEvents = bReplayable; }
-
 private:
 	/**
 	 * Context object for transactions that are to be processed.
@@ -160,7 +155,7 @@ private:
 	/**
 	 * Send a transaction finalized event.
 	 */
-	void SendTransactionFinalizedEvent(const FGuid& InTransactionId, const FGuid& InOperationId, UObject* InPrimaryObject, const TArray<FName>& InModifiedPackages, const TArray<FConcertExportedObject>& InObjectUpdates, const FConcertLocalIdentifierTable& InLocalIdentifierTable, const FText& InTitle, bool bReplayable);
+	void SendTransactionFinalizedEvent(const FGuid& InTransactionId, const FGuid& InOperationId, UObject* InPrimaryObject, const TArray<FName>& InModifiedPackages, const TArray<FConcertExportedObject>& InObjectUpdates, const FConcertLocalIdentifierTable& InLocalIdentifierTable, const FText& InTitle);
 
 	/**
 	 * Send a transaction snapshot event.
@@ -208,9 +203,4 @@ private:
 	 * Map of transaction IDs to the pending transaction that may be sent in the future (when finalized).
 	 */
 	TMap<FGuid, FPendingTransactionToSend> PendingTransactionsToSend;
-
-	/**
-	 * Whether the transaction events emitted to the server has the replayable flag on or off.
-	 */
-	bool bReplayableEvents = true;
 };
diff --git a/Engine/Plugins/Developer/Concert/ConcertSync/ConcertSyncClient/Source/ConcertSyncClient/Private/ConcertClientWorkspace.cpp b/Engine/Plugins/Developer/Concert/ConcertSync/ConcertSyncClient/Source/ConcertSyncClient/Private/ConcertClientWorkspace.cpp
--- a/Engine/Plugins/Developer/Concert/ConcertSync/ConcertSyncClient/Source/ConcertSyncClient/Private/ConcertClientWorkspace.cpp
+++ b/Engine/Plugins/Developer/Concert/ConcertSync/ConcertSyncClient/Source/ConcertSyncClient/Private/ConcertClientWorkspace.cpp
@@ -585,6 +585,13 @@ void FConcertClientWorkspace::OnEndFrame()
 		{
 			TransactionManager->ProcessPending();
 		}
+
+		if (bPendingStopIgnoringActivityOnRestore)
+		{
+			FConcertIgnoreActivityStateChangedEvent StateChangeEvent{LiveSession->GetSession().GetSessionClientEndpointId(), /*bIgnore*/false};
+			LiveSession->GetSession().SendCustomEvent(StateChangeEvent, LiveSession->GetSession().GetSessionServerEndpointId(), EConcertMessageFlags::ReliableOrdered);
+			bPendingStopIgnoringActivityOnRestore = false;
+		}
 	}
 }
 
@@ -763,12 +770,21 @@ void FConcertClientWorkspace::PostActivityUpdated(const FConcertSyncActivity& In
 	}
 }
 
-void FConcertClientWorkspace::SetEmittedEventsAsReplayable(bool bReplayable)
+void FConcertClientWorkspace::SetIgnoreOnRestoreFlagForEmittedActivities(bool bIgnore)
 {
-	check(TransactionManager && PackageManager);
-
-	TransactionManager->SetTransactionEventsAsReplayable(bReplayable);
-	PackageManager->SetPackageEventAsReplayable(bReplayable);
+	if (bIgnore) // Start ignoring further activities immediately.
+	{
+		FConcertIgnoreActivityStateChangedEvent StateChangeEvent{LiveSession->GetSession().GetSessionClientEndpointId(), bIgnore};
+		LiveSession->GetSession().SendCustomEvent(StateChangeEvent, LiveSession->GetSession().GetSessionServerEndpointId(), EConcertMessageFlags::ReliableOrdered);
+		bPendingStopIgnoringActivityOnRestore = false; // In case the 'ignore' state was toggle twice in the same frame.
+	}
+	else
+	{
+		// Stop 'ignoring' at the end of the frame, so that all pending transactions are sent and marked by the server properly. If a transaction started while
+		// the activities were set as 'ignored' but ended after the flat was cleared, the corresponding transaction activities will NOT be marked 'ignored' and
+		// will be seen as 'should restore' by the system.
+		bPendingStopIgnoringActivityOnRestore = true;
+	}
 }
 
 #undef LOCTEXT_NAMESPACE
diff --git a/Engine/Plugins/Developer/Concert/ConcertSync/ConcertSyncClient/Source/ConcertSyncClient/Private/ConcertClientWorkspace.h b/Engine/Plugins/Developer/Concert/ConcertSync/ConcertSyncClient/Source/ConcertSyncClient/Private/ConcertClientWorkspace.h
--- a/Engine/Plugins/Developer/Concert/ConcertSync/ConcertSyncClient/Source/ConcertSyncClient/Private/ConcertClientWorkspace.h
+++ b/Engine/Plugins/Developer/Concert/ConcertSync/ConcertSyncClient/Source/ConcertSyncClient/Private/ConcertClientWorkspace.h
@@ -49,7 +49,7 @@ public:
 	virtual FOnWorkspaceSynchronized& OnWorkspaceSynchronized() override;
 	virtual IConcertClientDataStore& GetDataStore() override;
 	virtual bool IsAssetModifiedByOtherClients(const FName& AssetName, int32* OutOtherClientsWithModifNum, TArray<FConcertClientInfo>* OutOtherClientsWithModifInfo, int32 OtherClientsWithModifMaxFetchNum) const override;
-	virtual void SetEmittedEventsAsReplayable(bool bEnabled) override;
+	virtual void SetIgnoreOnRestoreFlagForEmittedActivities(bool bIgnore) override;
 
 private:
 	/** Bind the workspace to this session. */
@@ -193,4 +193,7 @@ private:
 	
 	/** The session key/value store proxy. The real store is held by the server and shared across all clients. */
 	TUniquePtr<FConcertClientDataStore> DataStore;
+
+	/** True if the client has marked the further transaction as 'non-ignored'. This is sent at the end of the frame. */
+	bool bPendingStopIgnoringActivityOnRestore = false;
 };
diff --git a/Engine/Plugins/Developer/Concert/ConcertSync/ConcertSyncClient/Source/ConcertSyncClient/Public/IConcertClientTransactionBridge.h b/Engine/Plugins/Developer/Concert/ConcertSync/ConcertSyncClient/Source/ConcertSyncClient/Public/IConcertClientTransactionBridge.h
--- a/Engine/Plugins/Developer/Concert/ConcertSync/ConcertSyncClient/Source/ConcertSyncClient/Public/IConcertClientTransactionBridge.h
+++ b/Engine/Plugins/Developer/Concert/ConcertSync/ConcertSyncClient/Source/ConcertSyncClient/Public/IConcertClientTransactionBridge.h
@@ -30,7 +30,6 @@ struct FConcertClientLocalTransactionCommonData
 	TArray<FName> ModifiedPackages;
 	TArray<FConcertObjectId> ExcludedObjectUpdates;
 	bool bIsExcluded = false;
-	bool bReplayable = true;
 };
 
 /**
diff --git a/Engine/Plugins/Developer/Concert/ConcertSync/ConcertSyncClient/Source/ConcertSyncClient/Public/IConcertClientWorkspace.h b/Engine/Plugins/Developer/Concert/ConcertSync/ConcertSyncClient/Source/ConcertSyncClient/Public/IConcertClientWorkspace.h
--- a/Engine/Plugins/Developer/Concert/ConcertSync/ConcertSyncClient/Source/ConcertSyncClient/Public/IConcertClientWorkspace.h
+++ b/Engine/Plugins/Developer/Concert/ConcertSync/ConcertSyncClient/Source/ConcertSyncClient/Public/IConcertClientWorkspace.h
@@ -160,11 +160,10 @@ public:
 	virtual bool IsAssetModifiedByOtherClients(const FName& AssetName, int32* OutOtherClientsWithModifNum = nullptr, TArray<FConcertClientInfo>* OutOtherClientsWithModifInfo = nullptr, int32 OtherClientsWithModifMaxFetchNum = 0) const = 0;
 
 	/**
-	 * Controls whether the events emitted to the server through this workspace are marked as 'replayable'. By default, all events are marked as replayable. When
-	 * the workspace 'replayable' flag is false, the events emitted are recorded by the server, but are not replayed on the client(s). Non-replaylable events
-	 * are used for inspection rather than propagating/applying changes.
-	 * @note This was implemented to prevent multi-users transactions from being reapplied by disaster recovery in case of a crash during as multi-user session.
-	 * @param bReplayable True will mark all further events sent as replaylable, false as non-replayable.
+	 * Controls whether the activities emitted to the server through this workspace are marked as 'ignored on restore'. By default, all activities are marked as 'restorable'. When
+	 * the workspace 'ignore' state is true, the events emitted are recorded by the server, but are marked as 'should not restore'. Non-restorable activities are used for inspection.
+	 * @note This was implemented to prevent multi-users transactions from being restored by disaster recovery in case of a crash during as multi-user session.
+	 * @param bIgnore Whether all further events emitted are marked as 'ignored on restore'.
 	 */
-	virtual void SetEmittedEventsAsReplayable(bool bReplayable) = 0;
+	virtual void SetIgnoreOnRestoreFlagForEmittedActivities(bool bIgnore) = 0;
 };
diff --git a/Engine/Plugins/Developer/Concert/ConcertSync/ConcertSyncCore/Source/ConcertSyncCore/Public/ConcertTransactionEvents.h b/Engine/Plugins/Developer/Concert/ConcertSync/ConcertSyncCore/Source/ConcertSyncCore/Public/ConcertTransactionEvents.h
--- a/Engine/Plugins/Developer/Concert/ConcertSync/ConcertSyncCore/Source/ConcertSyncCore/Public/ConcertTransactionEvents.h
+++ b/Engine/Plugins/Developer/Concert/ConcertSync/ConcertSyncCore/Source/ConcertSyncCore/Public/ConcertTransactionEvents.h
@@ -129,10 +129,6 @@ struct FConcertTransactionEventBase
 
 	UPROPERTY()
 	TArray<FConcertExportedObject> ExportedObjects;
-
-	/** True to set the transaction event as replayable by clients. False would record it for logging/inspection purpose but will prevent clients from replaying it. */
-	UPROPERTY()
-	bool bReplayable = true;
 };
 
 USTRUCT()
diff --git a/Engine/Plugins/Developer/Concert/ConcertSync/ConcertSyncCore/Source/ConcertSyncCore/Public/ConcertWorkspaceMessages.h b/Engine/Plugins/Developer/Concert/ConcertSync/ConcertSyncCore/Source/ConcertSyncCore/Public/ConcertWorkspaceMessages.h
--- a/Engine/Plugins/Developer/Concert/ConcertSync/ConcertSyncCore/Source/ConcertSyncCore/Public/ConcertWorkspaceMessages.h
+++ b/Engine/Plugins/Developer/Concert/ConcertSync/ConcertSyncCore/Source/ConcertSyncCore/Public/ConcertWorkspaceMessages.h
@@ -73,10 +73,6 @@ struct FConcertPackageUpdateEvent
 
 	UPROPERTY()
 	FConcertPackage Package;
-
-	/** True to set the package event as replayable by clients. False would record it for logging/inspection purpose but will prevent clients from replaying it. */
-	UPROPERTY()
-	bool bReplayable = true;
 };
 
 USTRUCT()
@@ -163,3 +159,20 @@ struct FConcertPlaySessionEvent
 	UPROPERTY()
 	bool bIsSimulating = false;
 };
+
+/**
+ * Sets the specified client 'ignore on restore' state for further activities. The 'ignored' flag can be raised to mark a series of activities as 'should not be restored'.
+ * @note This can be used to record and monitor session activities for inspection purpose, for example allowing disaster recovery to record what
+ *       happens in a multi-user session without restoring such activities in case of crash (because they occurred in a transient sandbox).
+ */
+USTRUCT()
+struct FConcertIgnoreActivityStateChangedEvent
+{
+	GENERATED_BODY()
+
+	UPROPERTY()
+	FGuid EndpointId;
+
+	UPROPERTY()
+	bool bIgnore = false;
+};
diff --git a/Engine/Plugins/Developer/Concert/ConcertSync/ConcertSyncServer/Source/ConcertSyncServer/Private/ConcertServerWorkspace.cpp b/Engine/Plugins/Developer/Concert/ConcertSync/ConcertSyncServer/Source/ConcertSyncServer/Private/ConcertServerWorkspace.cpp
--- a/Engine/Plugins/Developer/Concert/ConcertSync/ConcertSyncServer/Source/ConcertSyncServer/Private/ConcertServerWorkspace.cpp
+++ b/Engine/Plugins/Developer/Concert/ConcertSync/ConcertSyncServer/Source/ConcertSyncServer/Private/ConcertServerWorkspace.cpp
@@ -56,6 +56,8 @@ void FConcertServerWorkspace::BindSession(const TSharedRef<FConcertSyncServerLiv
 	LiveSession->GetSession().RegisterCustomEventHandler<FConcertTransactionSnapshotEvent>(this, &FConcertServerWorkspace::HandleTransactionSnapshotEvent);
 
 	LiveSession->GetSession().RegisterCustomRequestHandler<FConcertResourceLockRequest, FConcertResourceLockResponse>(this, &FConcertServerWorkspace::HandleResourceLockRequest);
+
+	LiveSession->GetSession().RegisterCustomEventHandler<FConcertIgnoreActivityStateChangedEvent>(this, &FConcertServerWorkspace::HandleIgnoredActivityStateChanged);
 }
 
 void FConcertServerWorkspace::UnbindSession()
@@ -130,6 +132,7 @@ void FConcertServerWorkspace::HandleSessionClientChanged(IConcertServerSession&
 			ConnectionActivity.EndpointId = InClientInfo.ClientEndpointId;
 			ConnectionActivity.EventData.ConnectionEventType = EConcertSyncConnectionEventType::Connected;
 			ConnectionActivity.EventSummary.SetTypedPayload(FConcertSyncConnectionActivitySummary::CreateSummaryForEvent(ConnectionActivity.EventData));
+			ConnectionActivity.bIgnored = ShouldIgnoreClientActivityOnRestore(InClientInfo.ClientEndpointId);
 			AddConnectionActivity(ConnectionActivity);
 		}
 	}
@@ -145,12 +148,27 @@ void FConcertServerWorkspace::HandleSessionClientChanged(IConcertServerSession&
 			ConnectionActivity.EndpointId = InClientInfo.ClientEndpointId;
 			ConnectionActivity.EventData.ConnectionEventType = EConcertSyncConnectionEventType::Disconnected;
 			ConnectionActivity.EventSummary.SetTypedPayload(FConcertSyncConnectionActivitySummary::CreateSummaryForEvent(ConnectionActivity.EventData));
+			ConnectionActivity.bIgnored = ShouldIgnoreClientActivityOnRestore(InClientInfo.ClientEndpointId);
 			AddConnectionActivity(ConnectionActivity);
 		}
 
 		LiveSyncEndpoints.Remove(InClientInfo.ClientEndpointId);
 		ManualSyncEndpoints.Remove(InClientInfo.ClientEndpointId);
 		SyncCommandQueue->UnregisterEndpoint(InClientInfo.ClientEndpointId);
+
+		IgnoredActivityClients.Remove(InClientInfo.ClientEndpointId);
+	}
+}
+
+void FConcertServerWorkspace::HandleIgnoredActivityStateChanged(const FConcertSessionContext& Context, const FConcertIgnoreActivityStateChangedEvent& Event)
+{
+	if (Event.bIgnore)
+	{
+		IgnoredActivityClients.Add(Event.EndpointId);
+	}
+	else
+	{
+		IgnoredActivityClients.Remove(Event.EndpointId);
 	}
 }
 
@@ -280,7 +298,7 @@ void FConcertServerWorkspace::HandlePackageUpdateEvent(const FConcertSessionCont
 				}
 			}
 			PackageActivity.EventSummary.SetTypedPayload(FConcertSyncPackageActivitySummary::CreateSummaryForEvent(PackageActivity.EventData));
-			PackageActivity.bIgnored = !Event.bReplayable;
+			PackageActivity.bIgnored = ShouldIgnoreClientActivityOnRestore(Context.SourceEndpointId);
 			AddPackageActivity(PackageActivity);
 		}
 
@@ -323,7 +341,7 @@ void FConcertServerWorkspace::HandleTransactionFinalizedEvent(const FConcertSess
 			TransactionActivity.EndpointId = InEventContext.SourceEndpointId;
 			TransactionActivity.EventData.Transaction = InEvent;
 			TransactionActivity.EventSummary.SetTypedPayload(FConcertSyncTransactionActivitySummary::CreateSummaryForEvent(TransactionActivity.EventData));
-			TransactionActivity.bIgnored = !InEvent.bReplayable;
+			TransactionActivity.bIgnored = ShouldIgnoreClientActivityOnRestore(InEventContext.SourceEndpointId);
 			AddTransactionActivity(TransactionActivity);
 		}
 
@@ -342,7 +360,7 @@ void FConcertServerWorkspace::HandleTransactionFinalizedEvent(const FConcertSess
 
 void FConcertServerWorkspace::HandleTransactionSnapshotEvent(const FConcertSessionContext& InEventContext, const FConcertTransactionSnapshotEvent& InEvent)
 {
-	if (!EnumHasAnyFlags(LiveSession->GetSessionFlags(), EConcertSyncSessionFlags::EnableTransactions) || !InEvent.bReplayable)
+	if (!EnumHasAnyFlags(LiveSession->GetSessionFlags(), EConcertSyncSessionFlags::EnableTransactions))
 	{
 		return;
 	}
@@ -451,6 +469,7 @@ void FConcertServerWorkspace::HandleEndPlaySession(const FName InPlayPackageName
 		DummyPackageActivity.EventData.Package.Info.PackageName = InPlayPackageName;
 		LiveSession->GetSessionDatabase().GetTransactionMaxEventId(DummyPackageActivity.EventData.Package.Info.TransactionEventIdAtSave);
 		DummyPackageActivity.EventSummary.SetTypedPayload(FConcertSyncPackageActivitySummary::CreateSummaryForEvent(DummyPackageActivity.EventData));
+		DummyPackageActivity.bIgnored = ShouldIgnoreClientActivityOnRestore(InEndpointId);
 		AddPackageActivity(DummyPackageActivity);
 	}
 }
@@ -504,6 +523,7 @@ bool FConcertServerWorkspace::LockWorkspaceResource(const FName InResourceName,
 			LockActivity.EventData.LockEventType = EConcertSyncLockEventType::Locked;
 			LockActivity.EventData.ResourceNames.Add(InResourceName);
 			LockActivity.EventSummary.SetTypedPayload(FConcertSyncLockActivitySummary::CreateSummaryForEvent(LockActivity.EventData));
+			LockActivity.bIgnored = ShouldIgnoreClientActivityOnRestore(InLockEndpointId);
 			AddLockActivity(LockActivity);
 
 			FConcertResourceLockEvent LockEvent{ InLockEndpointId, {InResourceName}, EConcertResourceLockType::Lock };
@@ -561,6 +581,7 @@ bool FConcertServerWorkspace::LockWorkspaceResources(const TArray<FName>& InReso
 			LockActivity.EventData.LockEventType = EConcertSyncLockEventType::Locked;
 			LockActivity.EventData.ResourceNames = LockEvent.ResourceNames;
 			LockActivity.EventSummary.SetTypedPayload(FConcertSyncLockActivitySummary::CreateSummaryForEvent(LockActivity.EventData));
+			LockActivity.bIgnored = ShouldIgnoreClientActivityOnRestore(InLockEndpointId);
 			AddLockActivity(LockActivity);
 
 			LiveSession->GetSession().SendCustomEvent(LockEvent, LiveSession->GetSession().GetSessionClientEndpointIds(), EConcertMessageFlags::ReliableOrdered);
@@ -598,6 +619,7 @@ bool FConcertServerWorkspace::UnlockWorkspaceResource(const FName InResourceName
 					LockActivity.EventData.LockEventType = EConcertSyncLockEventType::Unlocked;
 					LockActivity.EventData.ResourceNames.Add(InResourceName);
 					LockActivity.EventSummary.SetTypedPayload(FConcertSyncLockActivitySummary::CreateSummaryForEvent(LockActivity.EventData));
+					LockActivity.bIgnored = ShouldIgnoreClientActivityOnRestore(InLockEndpointId);
 					AddLockActivity(LockActivity);
 
 					FConcertResourceLockEvent LockEvent{ InLockEndpointId, {InResourceName}, EConcertResourceLockType::Unlock };
@@ -652,6 +674,7 @@ bool FConcertServerWorkspace::UnlockWorkspaceResources(const TArray<FName>& InRe
 		LockActivity.EventData.LockEventType = EConcertSyncLockEventType::Unlocked;
 		LockActivity.EventData.ResourceNames = LockEvent.ResourceNames;
 		LockActivity.EventSummary.SetTypedPayload(FConcertSyncLockActivitySummary::CreateSummaryForEvent(LockActivity.EventData));
+		LockActivity.bIgnored = ShouldIgnoreClientActivityOnRestore(InLockEndpointId);
 		AddLockActivity(LockActivity);
 
 		LiveSession->GetSession().SendCustomEvent(LockEvent, LiveSession->GetSession().GetSessionClientEndpointIds(), EConcertMessageFlags::ReliableOrdered);
@@ -686,6 +709,7 @@ void FConcertServerWorkspace::UnlockAllWorkspaceResources(const FGuid& InLockEnd
 			LockActivity.EventData.LockEventType = EConcertSyncLockEventType::Unlocked;
 			LockActivity.EventData.ResourceNames = LockEvent.ResourceNames;
 			LockActivity.EventSummary.SetTypedPayload(FConcertSyncLockActivitySummary::CreateSummaryForEvent(LockActivity.EventData));
+			LockActivity.bIgnored = ShouldIgnoreClientActivityOnRestore(InLockEndpointId);
 			AddLockActivity(LockActivity);
 		}
 
diff --git a/Engine/Plugins/Developer/Concert/ConcertSync/ConcertSyncServer/Source/ConcertSyncServer/Private/ConcertServerWorkspace.h b/Engine/Plugins/Developer/Concert/ConcertSync/ConcertSyncServer/Source/ConcertSyncServer/Private/ConcertServerWorkspace.h
--- a/Engine/Plugins/Developer/Concert/ConcertSync/ConcertSyncServer/Source/ConcertSyncServer/Private/ConcertServerWorkspace.h
+++ b/Engine/Plugins/Developer/Concert/ConcertSync/ConcertSyncServer/Source/ConcertSyncServer/Private/ConcertServerWorkspace.h
@@ -59,6 +59,9 @@ private:
 	/** */
 	EConcertSessionResponseCode HandleResourceLockRequest(const FConcertSessionContext& Context, const FConcertResourceLockRequest& Request, FConcertResourceLockResponse& Response);
 
+	/** */
+	void HandleIgnoredActivityStateChanged(const FConcertSessionContext& Context, const FConcertIgnoreActivityStateChangedEvent& Event);
+
 	/** Invoked when the client corresponding to the specified endpoint begins to "Play" in a mode such as PIE or SIE. */
 	void HandleBeginPlaySession(const FName InPlayPackageName, const FGuid& InEndpointId, bool bIsSimulating);
 
@@ -215,6 +218,13 @@ private:
 	 */
 	void PostActivityAdded(const int64 InActivityId);
 
+	/**
+	 * Check if a given client was configured to prevent restoring the activities being recorded.
+	 * @param ClientEndpoint The client for which the check must be done.
+	 * @return True if the activity should be flagged as 'ignored on restore'.
+	 */
+	bool ShouldIgnoreClientActivityOnRestore(const FGuid ClientEndpoint) const { return IgnoredActivityClients.Contains(ClientEndpoint); }
+
 	/** Live session tracked by this workspace */
 	TSharedPtr<FConcertSyncServerLiveSession> LiveSession;
 
@@ -250,4 +260,7 @@ private:
 
 	/** The data store shared by all clients connected to the server tracked by this workspace. */
 	TUniquePtr<FConcertServerDataStore> DataStore;
+
+	/** Contains the list of clients for which activities should be marked as 'ignore on restore'. @see FConcertIgnoreActivityStateChangedEvent. */
+	TSet<FGuid> IgnoredActivityClients;
 };
diff --git a/Engine/Plugins/Developer/Concert/DisasterRecoveryClient/Source/DisasterRecoveryClient/Private/DisasterRecoveryClientModule.cpp b/Engine/Plugins/Developer/Concert/DisasterRecoveryClient/Source/DisasterRecoveryClient/Private/DisasterRecoveryClientModule.cpp
--- a/Engine/Plugins/Developer/Concert/DisasterRecoveryClient/Source/DisasterRecoveryClient/Private/DisasterRecoveryClientModule.cpp
+++ b/Engine/Plugins/Developer/Concert/DisasterRecoveryClient/Source/DisasterRecoveryClient/Private/DisasterRecoveryClientModule.cpp
@@ -161,13 +161,13 @@ private:
 	void HandleSyncSessionStartup(const IConcertSyncClient* SyncClient)
 	{
 		check(DisasterRecoveryClient.Get() != SyncClient)
-		SetDisasterRecoveryEventsAsReplayable(IsCompatibleWithOtherConcertSessions(SyncClient, /*SyncClientShuttingDownSession*/nullptr));
+		SetIgnoreOnRestoreState(!IsCompatibleWithOtherConcertSessions(SyncClient, /*SyncClientShuttingDownSession*/nullptr));
 	}
 
 	void HandleSyncSessionShutdown(const IConcertSyncClient* SyncClient)
 	{
 		check(DisasterRecoveryClient.Get() != SyncClient);
-		SetDisasterRecoveryEventsAsReplayable(IsCompatibleWithOtherConcertSessions(/*SyncClientStartingSession*/nullptr, SyncClient));
+		SetIgnoreOnRestoreState(!IsCompatibleWithOtherConcertSessions(/*SyncClientStartingSession*/nullptr, SyncClient));
 	}
 
 	FGuid GetDisasterRecoverySessionId() const
@@ -411,7 +411,7 @@ private:
 		DisasterRecoveryClient->Startup(ClientConfig, EConcertSyncSessionFlags::Default_DisasterRecoverySession);
 
 		// Set all events captured by the disaster recovery service as 'replayable' unless another concert client (assumed Multi-User) has created an incompatible session.
-		SetDisasterRecoveryEventsAsReplayable(IsCompatibleWithOtherConcertSessions(/*SyncClientStartingSession*/nullptr, /*SyncClientShuttingDownSession*/nullptr));
+		SetIgnoreOnRestoreState(!IsCompatibleWithOtherConcertSessions(/*SyncClientStartingSession*/nullptr, /*SyncClientShuttingDownSession*/nullptr));
 
 		// If something needs to be recovered from a crash.
 		if (RestoringSession)
@@ -503,12 +503,12 @@ private:
 		return SyncClient->GetConcertClient()->GetRole() != TEXT("MultiUser");
 	}
 
-	/** Sets whether further Concert events (transaction/package) emitted by Disaster Recovery have the replayable flag on or off. */
-	void SetDisasterRecoveryEventsAsReplayable(bool bReplayable)
+	/** Sets whether further Concert events (transaction/package) emitted by Disaster Recovery have the 'ignore' flag on or off. */
+	void SetIgnoreOnRestoreState(bool bIgnore)
 	{
 		if (TSharedPtr<IConcertClientWorkspace> Workspace = DisasterRecoveryClient ? DisasterRecoveryClient->GetWorkspace() : TSharedPtr<IConcertClientWorkspace>())
 		{
-			Workspace->SetEmittedEventsAsReplayable(bReplayable);
+			Workspace->SetIgnoreOnRestoreFlagForEmittedActivities(bIgnore);
 		}
 	}
 
