Fixed DisasterRecovery asserting if the recovery session database could not be open or was corrupted.
  - Changed the recovery session activity stream to return 'I'm done' when it encounters an error to prevent the caller to ignore the error and continue asking new activities (which triggers the assert).
  - Handled the error reported by the stream when searching for a recoverable session just after the Editor finished its initialization.

Fixed DisasterRecovery freezing the Editor if the new session database could not be created successfully.
  - Reported the error to the client instead of silently ignoring it.
  - Updated the verbose log to Warning to let the user know that a new session could not be created (and therefore disaster recovery is going to be off.

Details:
  - The bug UE-89128 is possibly related to UE-89826 (at least on Mac) where the database file would fails to open successfully. Corrupted database/broken file API would fail to open the database file and will result in the assert below:
        Assertion failed: (RequestFetchCount < 0 && LowestFetchedActivityId == 0) || (RequestFetchCount > 0 && LowestFetchedActivityId > 1) [File:[REDACTED_USER_HOME]/Build/++UE4+Licensee/Sync/Engine/Plugins/Developer/Concert/ConcertSync/ConcertSyncClient/Source/ConcertSyncClient/Private/ConcertActivityStream.cpp] [Line: 66]

This CL fix can be tested by:
      1. Launch the Editor, generate a crash with 'debug crash' command.
      2. Comment 'SessionDatabase->Open(Session->GetSessionWorkingDirectory());' in Engine\Plugins\Developer\Concert\ConcertSync\ConcertSyncServer\Source\ConcertSyncServer\Private\ConcertSyncServerLiveSession.cpp to simulate SQLite failing to open the database.
      3. Recompile CrashReportClientEditor(Windows) or UnrealRecoverySvc(Mac/Linux)
      4. Relaunch the editor.

diff --git a/Engine/Plugins/Developer/Concert/ConcertMain/Source/Concert/Private/ConcertServer.cpp b/Engine/Plugins/Developer/Concert/ConcertMain/Source/Concert/Private/ConcertServer.cpp
--- a/Engine/Plugins/Developer/Concert/ConcertMain/Source/Concert/Private/ConcertServer.cpp
+++ b/Engine/Plugins/Developer/Concert/ConcertMain/Source/Concert/Private/ConcertServer.cpp
@@ -1307,7 +1307,7 @@ TFuture<FConcertAdmin_GetSessionActivitiesResponse> FConcertServer::HandleGetSes
 	else // The only reason to get here is when the session is not found.
 	{
 		ResponseData.ResponseCode = EConcertResponseCode::Failed;
-		ResponseData.Reason = LOCTEXT("Error_SessionActivities_SessionDoesNotExist", "Session does not exist.");
+		ResponseData.Reason = LOCTEXT("Error_SessionActivities_SessionDoesNotExist", "Session does not exist or its database is corrupted.");
 		UE_LOG(LogConcert, Display, TEXT("Failed to fetch activities from session (Id: %s, Reason: %s)"), *Message->SessionId.ToString(), *ResponseData.Reason.ToString());
 	}
 
@@ -1373,11 +1373,14 @@ TSharedPtr<IConcertServerSession> FConcertServer::CreateLiveSession(const FConce
 		InRepository.GetSessionWorkingDir(LiveSessionInfo.SessionId)
 		);
 
-	LiveSessions.Add(LiveSessionInfo.SessionId, LiveSession);
-	EventSink->OnLiveSessionCreated(*this, LiveSession.ToSharedRef());
-	LiveSession->Startup();
+	if (EventSink->OnLiveSessionCreated(*this, LiveSession.ToSharedRef())) // EventSync could complete the session initialization?
+	{
+		LiveSessions.Add(LiveSessionInfo.SessionId, LiveSession);
+		LiveSession->Startup();
+		return LiveSession;
+	}
 
-	return LiveSession;
+	return nullptr;
 }
 
 bool FConcertServer::DestroyLiveSession(const FGuid& LiveSessionId, const bool bDeleteSessionData)
@@ -1440,9 +1443,7 @@ bool FConcertServer::CreateArchivedSession(const FConcertSessionInfo& SessionInf
 	check(!ArchivedSessions.Contains(SessionInfo.SessionId) && !GetArchivedSessionIdByName(SessionInfo.SessionName).IsValid());
 
 	ArchivedSessions.Add(SessionInfo.SessionId, SessionInfo);
-	EventSink->OnArchivedSessionCreated(*this, GetSessionSavedDir(SessionInfo.SessionId), SessionInfo);
-
-	return true;
+	return EventSink->OnArchivedSessionCreated(*this, GetSessionSavedDir(SessionInfo.SessionId), SessionInfo);
 }
 
 bool FConcertServer::DestroyArchivedSession(const FGuid& ArchivedSessionId, const bool bDeleteSessionData)
diff --git a/Engine/Plugins/Developer/Concert/ConcertMain/Source/Concert/Public/IConcertServerEventSink.h b/Engine/Plugins/Developer/Concert/ConcertMain/Source/Concert/Public/IConcertServerEventSink.h
--- a/Engine/Plugins/Developer/Concert/ConcertMain/Source/Concert/Public/IConcertServerEventSink.h
+++ b/Engine/Plugins/Developer/Concert/ConcertMain/Source/Concert/Public/IConcertServerEventSink.h
@@ -23,8 +23,9 @@ public:
 	/**
 	 * Called after the session has been created (and before Startup has been called on it).
 	 * @note This function is called for both newly created sessions and after recovering a live session during server start-up.
+	 * @return true if the session creation could be completed without error, false otherwise (ex if the database fails to open).
 	 */
-	virtual void OnLiveSessionCreated(const IConcertServer& InServer, TSharedRef<IConcertServerSession> InLiveSession) = 0;
+	virtual bool OnLiveSessionCreated(const IConcertServer& InServer, TSharedRef<IConcertServerSession> InLiveSession) = 0;
 	
 	/**
 	 * Called before the session is destroyed (and before Shutdown is called on it).
@@ -35,8 +36,9 @@ public:
 	/**
 	 * Called after the session has been created.
 	 * @note This function is called for both newly created sessions and after recovering an archived session during server start-up.
+	 * @return true if the session creation could be completed without error, false otherwise (ex if the database fails to open).
 	 */
-	virtual void OnArchivedSessionCreated(const IConcertServer& InServer, const FString& InArchivedSessionRoot, const FConcertSessionInfo& InArchivedSessionInfo) = 0;
+	virtual bool OnArchivedSessionCreated(const IConcertServer& InServer, const FString& InArchivedSessionRoot, const FConcertSessionInfo& InArchivedSessionInfo) = 0;
 	
 	/**
 	 * Called before the session is destroyed.
diff --git a/Engine/Plugins/Developer/Concert/ConcertSync/ConcertSyncClient/Source/ConcertSyncClient/Private/ConcertActivityStream.cpp b/Engine/Plugins/Developer/Concert/ConcertSync/ConcertSyncClient/Source/ConcertSyncClient/Private/ConcertActivityStream.cpp
--- a/Engine/Plugins/Developer/Concert/ConcertSync/ConcertSyncClient/Source/ConcertSyncClient/Private/ConcertActivityStream.cpp
+++ b/Engine/Plugins/Developer/Concert/ConcertSync/ConcertSyncClient/Source/ConcertSyncClient/Private/ConcertActivityStream.cpp
@@ -26,8 +26,8 @@ bool FConcertActivityStream::Read(TArray<TSharedPtr<FConcertClientSessionActivit
 	if (!OutErrorMsg.IsEmpty())
 	{
 		OutReadCount = 0;
-		LastErrorMsg = FText();
-		return false; // Not done yet. This allow the caller to deal with the error and decide to continue or not. He will most likely abort, but lets him decide.
+		bAllActivitiesFetched = true;
+		return true; // Done? When an error occurs, the stream is in a unknown state, consider the streaming done.
 	}
 
 	if (!bAllActivitiesFetched)
diff --git a/Engine/Plugins/Developer/Concert/ConcertSync/ConcertSyncServer/Source/ConcertSyncServer/Private/ConcertSyncServer.cpp b/Engine/Plugins/Developer/Concert/ConcertSync/ConcertSyncServer/Source/ConcertSyncServer/Private/ConcertSyncServer.cpp
--- a/Engine/Plugins/Developer/Concert/ConcertSync/ConcertSyncServer/Source/ConcertSyncServer/Private/ConcertSyncServer.cpp
+++ b/Engine/Plugins/Developer/Concert/ConcertSync/ConcertSyncServer/Source/ConcertSyncServer/Private/ConcertSyncServer.cpp
@@ -289,10 +289,10 @@ void FConcertSyncServer::GetSessionsFromPath(const IConcertServer& InServer, con
 	});
 }
 
-void FConcertSyncServer::OnLiveSessionCreated(const IConcertServer& InServer, TSharedRef<IConcertServerSession> InSession)
+bool FConcertSyncServer::OnLiveSessionCreated(const IConcertServer& InServer, TSharedRef<IConcertServerSession> InSession)
 {
 	ConcertSyncServerUtils::WriteSessionInfoToDirectory(InSession->GetSessionWorkingDirectory(), InSession->GetSessionInfo());
-	CreateLiveSession(InSession);
+	return CreateLiveSession(InSession);
 }
 
 void FConcertSyncServer::OnLiveSessionDestroyed(const IConcertServer& InServer, TSharedRef<IConcertServerSession> InSession)
@@ -300,10 +300,10 @@ void FConcertSyncServer::OnLiveSessionDestroyed(const IConcertServer& InServer,
 	DestroyLiveSession(InSession);
 }
 
-void FConcertSyncServer::OnArchivedSessionCreated(const IConcertServer& InServer, const FString& InArchivedSessionRoot, const FConcertSessionInfo& InArchivedSessionInfo)
+bool FConcertSyncServer::OnArchivedSessionCreated(const IConcertServer& InServer, const FString& InArchivedSessionRoot, const FConcertSessionInfo& InArchivedSessionInfo)
 {
 	ConcertSyncServerUtils::WriteSessionInfoToDirectory(InArchivedSessionRoot, InArchivedSessionInfo);
-	CreateArchivedSession(InArchivedSessionRoot, InArchivedSessionInfo);
+	return CreateArchivedSession(InArchivedSessionRoot, InArchivedSessionInfo);
 }
 
 void FConcertSyncServer::OnArchivedSessionDestroyed(const IConcertServer& InServer, const FGuid& InArchivedSessionId)
@@ -490,7 +490,7 @@ void FConcertSyncServer::DestroySequencerManager(const TSharedRef<FConcertSyncSe
 	LiveSessionSequencerManagers.Remove(InLiveSession->GetSession().GetId());
 }
 
-void FConcertSyncServer::CreateLiveSession(const TSharedRef<IConcertServerSession>& InSession)
+bool FConcertSyncServer::CreateLiveSession(const TSharedRef<IConcertServerSession>& InSession)
 {
 	DestroyLiveSession(InSession);
 
@@ -503,7 +503,11 @@ void FConcertSyncServer::CreateLiveSession(const TSharedRef<IConcertServerSessio
 		{
 			CreateSequencerManager(LiveSession.ToSharedRef());
 		}
+
+		return true;
 	}
+
+	return false;
 }
 
 void FConcertSyncServer::DestroyLiveSession(const TSharedRef<IConcertServerSession>& InSession)
@@ -516,7 +520,7 @@ void FConcertSyncServer::DestroyLiveSession(const TSharedRef<IConcertServerSessi
 	}
 }
 
-void FConcertSyncServer::CreateArchivedSession(const FString& InArchivedSessionRoot, const FConcertSessionInfo& InArchivedSessionInfo)
+bool FConcertSyncServer::CreateArchivedSession(const FString& InArchivedSessionRoot, const FConcertSessionInfo& InArchivedSessionInfo)
 {
 	DestroyArchivedSession(InArchivedSessionInfo.SessionId);
 
@@ -524,7 +528,10 @@ void FConcertSyncServer::CreateArchivedSession(const FString& InArchivedSessionR
 	if (ArchivedSession->IsValidSession())
 	{
 		ArchivedSessions.Add(ArchivedSession->GetId(), ArchivedSession);
+		return true;
 	}
+
+	return false;
 }
 
 void FConcertSyncServer::DestroyArchivedSession(const FGuid& InArchivedSessionId)
diff --git a/Engine/Plugins/Developer/Concert/ConcertSync/ConcertSyncServer/Source/ConcertSyncServer/Private/ConcertSyncServer.h b/Engine/Plugins/Developer/Concert/ConcertSync/ConcertSyncServer/Source/ConcertSyncServer/Private/ConcertSyncServer.h
--- a/Engine/Plugins/Developer/Concert/ConcertSync/ConcertSyncServer/Source/ConcertSyncServer/Private/ConcertSyncServer.h
+++ b/Engine/Plugins/Developer/Concert/ConcertSync/ConcertSyncServer/Source/ConcertSyncServer/Private/ConcertSyncServer.h
@@ -31,9 +31,9 @@ public:
 
 	//~ IConcertServerEventSink interface
 	virtual void GetSessionsFromPath(const IConcertServer& InServer, const FString& InPath, TArray<FConcertSessionInfo>& OutSessionInfos, TArray<FDateTime>* OutSessionCreationTimes = nullptr) override;
-	virtual void OnLiveSessionCreated(const IConcertServer& InServer, TSharedRef<IConcertServerSession> InLiveSession) override;
+	virtual bool OnLiveSessionCreated(const IConcertServer& InServer, TSharedRef<IConcertServerSession> InLiveSession) override;
 	virtual void OnLiveSessionDestroyed(const IConcertServer& InServer, TSharedRef<IConcertServerSession> InLiveSession) override;
-	virtual void OnArchivedSessionCreated(const IConcertServer& InServer, const FString& InArchivedSessionRoot, const FConcertSessionInfo& InArchivedSessionInfo) override;
+	virtual bool OnArchivedSessionCreated(const IConcertServer& InServer, const FString& InArchivedSessionRoot, const FConcertSessionInfo& InArchivedSessionInfo) override;
 	virtual void OnArchivedSessionDestroyed(const IConcertServer& InServer, const FGuid& InArchivedSessionId) override;
 	virtual bool ArchiveSession(const IConcertServer& InServer, TSharedRef<IConcertServerSession> InLiveSession, const FString& InArchivedSessionRoot, const FConcertSessionInfo& InArchivedSessionInfo, const FConcertSessionFilter& InSessionFilter) override;
 	virtual bool ArchiveSession(const IConcertServer& InServer, const FString& InLiveSessionWorkingDir, const FString& InArchivedSessionRoot, const FConcertSessionInfo& InArchivedSessionInfo, const FConcertSessionFilter& InSessionFilter) override;
@@ -51,10 +51,10 @@ private:
 	void CreateSequencerManager(const TSharedRef<FConcertSyncServerLiveSession>& InLiveSession);
 	void DestroySequencerManager(const TSharedRef<FConcertSyncServerLiveSession>& InLiveSession);
 
-	void CreateLiveSession(const TSharedRef<IConcertServerSession>& InSession);
+	bool CreateLiveSession(const TSharedRef<IConcertServerSession>& InSession);
 	void DestroyLiveSession(const TSharedRef<IConcertServerSession>& InSession);
 
-	void CreateArchivedSession(const FString& InArchivedSessionRoot, const FConcertSessionInfo& InArchivedSessionInfo);
+	bool CreateArchivedSession(const FString& InArchivedSessionRoot, const FConcertSessionInfo& InArchivedSessionInfo);
 	void DestroyArchivedSession(const FGuid& InArchivedSessionId);
 
 	bool GetSessionActivities(const FConcertSyncSessionDatabase& Database, int64 FromActivityId, int64 ActivityCount, TArray<FConcertSessionSerializedPayload>& OutActivities, TMap<FGuid, FConcertClientInfo>& OutEndpointClientInfoMap, bool bIncludeDetails);
diff --git a/Engine/Plugins/Developer/Concert/DisasterRecoveryClient/Source/DisasterRecoveryClient/Private/DisasterRecoveryTasks.cpp b/Engine/Plugins/Developer/Concert/DisasterRecoveryClient/Source/DisasterRecoveryClient/Private/DisasterRecoveryTasks.cpp
--- a/Engine/Plugins/Developer/Concert/DisasterRecoveryClient/Source/DisasterRecoveryClient/Private/DisasterRecoveryTasks.cpp
+++ b/Engine/Plugins/Developer/Concert/DisasterRecoveryClient/Source/DisasterRecoveryClient/Private/DisasterRecoveryTasks.cpp
@@ -292,7 +292,7 @@ void FCreateAndJoinDisasterRecoverySessionTask::Start()
 			{
 				OnErrorFn(DisasterRecoveryTasksUtil::GetSystemErrorMessage(Response));
 				bDone = true;
-				UE_LOG(LogDisasterRecovery, Verbose, TEXT("...Creating/Restoring session '%s' -> Failed"), *SessionName);
+				UE_LOG(LogDisasterRecovery, Warning, TEXT("Failed to create or restore session %s. Disaster recovery will not be enabled for this session."), *SessionName);
 			}
 		}
 	});
@@ -425,9 +425,9 @@ bool FFindRecoverableDisasterRecoverySessionTask::Tick()
 	{
 		int32 ReadCount;
 		FText ErrorMsg;
-		bool bEndOfStream = ActivityStream->Read(Activities, ReadCount, ErrorMsg); // Read the next batch of activities.
+		bool bEndOfStreamOrError = ActivityStream->Read(Activities, ReadCount, ErrorMsg); // Read activities (if any)
 
-		if (ReadCount > 0)
+		if (ReadCount > 0) // On error, ReadCount is 0.
 		{
 			for (const TSharedPtr<FConcertClientSessionActivity>& Activity : Activities)
 			{
@@ -435,16 +435,23 @@ bool FFindRecoverableDisasterRecoverySessionTask::Tick()
 				{
 					OnCandidateSearchResultFn(CurrentCandidate); // Found a candidate with recoverable activities, transit to the 'exit' state.
 					UE_LOG(LogDisasterRecovery, Verbose, TEXT("...Looking up a recoverable session -> Found recoverable session '%s' among the candidates"), *CurrentCandidate->SessionName);
-					return true; // Done­.
+					return true; // Found a recoverable candidate.
 				}
 			}
 			Activities.Reset(); // Don't need to keep the activities.
 		}
 
-		if (bEndOfStream)
+		if (bEndOfStreamOrError)
 		{
-			// No activities to recover, transit to the 'pick a candidate' state.
-			UE_LOG(LogDisasterRecovery, Verbose, TEXT("Eliminated candidate session %s. No recoverable activities available."), *CurrentCandidate->SessionName);
+			if (ErrorMsg.IsEmpty()) // No error.
+			{
+				UE_LOG(LogDisasterRecovery, Verbose, TEXT("Eliminated candidate session %s. No recoverable activities available."), *CurrentCandidate->SessionName);
+			}
+			else
+			{
+				UE_LOG(LogDisasterRecovery, Warning, TEXT("Eliminated candidate session %s. Reason: %s"), *CurrentCandidate->SessionName, *ErrorMsg.ToString());
+			}
+
 			Candidates.Pop();
 			ActivityStream.Reset();
 			CurrentCandidate.Reset();
