cleaning up the previous change to FThreadStats::StopThread() that caused it to be called twice, and crashing on the second call as TaskGraph was already shut down

many programs were calling FThreadStats::StopThread() and FTaskGraphInterfaces::Shutdown() explicitly despite they are called from FEngineLoop::AppExit(). Calling them twice is problematic because FThreadStats::StopThread() expects TaskGraph to be still alive. This didn't cause any visible problems before because FThreadStats::StopThread() wasn't called in AppExit().

diff --git a/Engine/Plugins/Developer/Concert/ConcertSync/ConcertSyncServer/Source/ConcertSyncServer/Public/ConcertSyncServerLoop.inl b/Engine/Plugins/Developer/Concert/ConcertSync/ConcertSyncServer/Source/ConcertSyncServer/Public/ConcertSyncServerLoop.inl
--- a/Engine/Plugins/Developer/Concert/ConcertSync/ConcertSyncServer/Source/ConcertSyncServer/Public/ConcertSyncServerLoop.inl
+++ b/Engine/Plugins/Developer/Concert/ConcertSync/ConcertSyncServer/Source/ConcertSyncServer/Public/ConcertSyncServerLoop.inl
@@ -191,10 +191,6 @@ int32 ConcertSyncServerLoop(int32 ArgC, TCHAR** ArgV, const FConcertSyncServerLo
 
 	// Unloading Modules isn't handled by AppExit
 	FModuleManager::Get().UnloadModulesAtShutdown();
-	// Nor is it handling stats, if any
-#if STATS
-	FThreadStats::StopThread();
-#endif
 
 	FEngineLoop::AppExit();
 
diff --git a/Engine/Source/Programs/UnrealFrontend/Private/UnrealFrontendMain.cpp b/Engine/Source/Programs/UnrealFrontend/Private/UnrealFrontendMain.cpp
--- a/Engine/Source/Programs/UnrealFrontend/Private/UnrealFrontendMain.cpp
+++ b/Engine/Source/Programs/UnrealFrontend/Private/UnrealFrontendMain.cpp
@@ -94,11 +94,5 @@ int32 UnrealFrontendMain( const TCHAR* CommandLine )
 	FEngineLoop::AppPreExit();
 	FModuleManager::Get().UnloadModulesAtShutdown();
 
-#if STATS
-	FThreadStats::StopThread();
-#endif
-
-	FTaskGraphInterface::Shutdown();
-
 	return Succeeded ? 0 : -1;
 }
diff --git a/Engine/Source/Programs/UnrealVersionSelector/Private/UnrealVersionSelector.cpp b/Engine/Source/Programs/UnrealVersionSelector/Private/UnrealVersionSelector.cpp
--- a/Engine/Source/Programs/UnrealVersionSelector/Private/UnrealVersionSelector.cpp
+++ b/Engine/Source/Programs/UnrealVersionSelector/Private/UnrealVersionSelector.cpp
@@ -448,11 +448,6 @@ int Main(const TArray<FString>& Arguments)
 		FEngineLoop::AppPreExit();
 		FModuleManager::Get().UnloadModulesAtShutdown();
 
-	#if STATS
-		FThreadStats::StopThread();
-	#endif
-
-		FTaskGraphInterface::Shutdown();
 		return Result;
 	}
 
