Add a deprecation warning when using the ConcertVersion flag.

diff --git a/Engine/Plugins/Developer/Concert/ConcertSync/ConcertSyncServer/Source/ConcertSyncServer/Private/ConcertSyncServerModule.cpp b/Engine/Plugins/Developer/Concert/ConcertSync/ConcertSyncServer/Source/ConcertSyncServer/Private/ConcertSyncServerModule.cpp
--- a/Engine/Plugins/Developer/Concert/ConcertSync/ConcertSyncServer/Source/ConcertSyncServer/Private/ConcertSyncServerModule.cpp
+++ b/Engine/Plugins/Developer/Concert/ConcertSync/ConcertSyncServer/Source/ConcertSyncServer/Private/ConcertSyncServerModule.cpp
@@ -41,6 +41,7 @@ public:
 			if (FParse::Value(CommandLine, TEXT("-CONCERTVERSION="), VersionString))
 			{
 				FEngineVersion EngineVersion;
+				UE_LOG(LogConcert, Warning, TEXT("CONCERTVERSION command line flag is deprecated and will be removed in a future version."));
 				if (FEngineVersion::Parse(VersionString, EngineVersion))
 				{
 					ServerConfig->DefaultVersionInfo.EngineVersion.Initialize(EngineVersion);
