Expose important server events as multicast delegates

diff --git a/Engine/Plugins/Developer/Concert/ConcertMain/Source/Concert/Private/ConcertServerEvents.cpp b/Engine/Plugins/Developer/Concert/ConcertMain/Source/Concert/Private/ConcertServerEvents.cpp
new file mode 100644
--- /dev/null
+++ b/Engine/Plugins/Developer/Concert/ConcertMain/Source/Concert/Private/ConcertServerEvents.cpp
@@ -0,0 +1,69 @@
+﻿// Copyright Epic Games, Inc. All Rights Reserved.
+
+#include "ConcertServerEvents.h"
+
+ConcertServerEvents::FOnLiveSessionCreated& ConcertServerEvents::OnLiveSessionCreated()
+{
+	static FOnLiveSessionCreated Event;
+	return Event;
+}
+
+ConcertServerEvents::FOnArchivedSessionCreated& ConcertServerEvents::OnArchivedSessionCreated()
+{
+	static FOnArchivedSessionCreated Event;
+	return Event;
+}
+
+ConcertServerEvents::FOnLiveSessionDestroyed& ConcertServerEvents::OnLiveSessionDestroyed()
+{
+	static FOnLiveSessionDestroyed Event;
+	return Event;
+}
+
+ConcertServerEvents::FOnArchivedSessionDestroyed& ConcertServerEvents::OnArchivedSessionDestroyed()
+{
+	static FOnArchivedSessionDestroyed Event;
+	return Event;
+}
+
+ConcertServerEvents::FArchiveSession_WithSession& ConcertServerEvents::ArchiveSession_WithSession()
+{
+	static FArchiveSession_WithSession Event;
+	return Event;
+}
+
+ConcertServerEvents::FArchiveSession_WithWorkingDir& ConcertServerEvents::ArchiveSession_WithWorkingDir()
+{
+	static FArchiveSession_WithWorkingDir Event;
+	return Event;
+}
+
+ConcertServerEvents::FCopySession& ConcertServerEvents::CopySession()
+{
+	static FCopySession Event;
+	return Event;
+}
+
+ConcertServerEvents::FExportSession& ConcertServerEvents::ExportSession()
+{
+	static FExportSession Event;
+	return Event;
+}
+
+ConcertServerEvents::FRestoreSession& ConcertServerEvents::RestoreSession()
+{
+	static FRestoreSession Event;
+	return Event;
+}
+
+ConcertServerEvents::FOnLiveSessionRenamed& ConcertServerEvents::OnLiveSessionRenamed()
+{
+	static FOnLiveSessionRenamed Event;
+	return Event;
+}
+
+ConcertServerEvents::FOnArchivedSessionRenamed& ConcertServerEvents::OnArchivedSessionRenamed()
+{
+	static FOnArchivedSessionRenamed Event;
+	return Event;
+}
\ No newline at end of file
diff --git a/Engine/Plugins/Developer/Concert/ConcertMain/Source/Concert/Public/ConcertServerEvents.h b/Engine/Plugins/Developer/Concert/ConcertMain/Source/Concert/Public/ConcertServerEvents.h
new file mode 100644
--- /dev/null
+++ b/Engine/Plugins/Developer/Concert/ConcertMain/Source/Concert/Public/ConcertServerEvents.h
@@ -0,0 +1,40 @@
+﻿// Copyright Epic Games, Inc. All Rights Reserved.
+
+#pragma once
+
+#include "CoreMinimal.h"
+#include "ConcertMessageData.h"
+
+class IConcertServer;
+class IConcertServerSession;
+
+// Most events from IConcertServerEventSink are mirroed here
+namespace ConcertServerEvents
+{
+	DECLARE_MULTICAST_DELEGATE_ThreeParams(FOnLiveSessionCreated, bool /*bSuccess*/, const IConcertServer& /*InServer*/, TSharedRef<IConcertServerSession> /*InLiveSession*/);
+	DECLARE_MULTICAST_DELEGATE_FourParams(FOnArchivedSessionCreated, bool /*bSuccess*/, const IConcertServer& /*InServer*/, const FString& /*InArchivedSessionRoot*/, const FConcertSessionInfo& /*InArchivedSessionInfo*/);
+	
+	DECLARE_MULTICAST_DELEGATE_TwoParams(FOnLiveSessionDestroyed, const IConcertServer& /*InServer*/, TSharedRef<IConcertServerSession> /*InLiveSession*/);
+	DECLARE_MULTICAST_DELEGATE_TwoParams(FOnArchivedSessionDestroyed, const IConcertServer& /*InServer*/, const FGuid& /*InArchivedSessionId*/);
+	DECLARE_MULTICAST_DELEGATE_SixParams(FArchiveSession_WithSession, bool /*bSuccess*/, const IConcertServer& /*InServer*/, TSharedRef<IConcertServerSession> /*InLiveSession*/, const FString& /*InArchivedSessionRoot*/, const FConcertSessionInfo& /*InArchivedSessionInfo*/, const FConcertSessionFilter& /*InSessionFilter*/);
+	DECLARE_MULTICAST_DELEGATE_SixParams(FArchiveSession_WithWorkingDir, bool /*bSuccess*/, const IConcertServer& /*InServer*/, const FString& /*InLiveSessionWorkingDir*/, const FString& /*InArchivedSessionRoot*/, const FConcertSessionInfo& /*InArchivedSessionInfo*/, const FConcertSessionFilter& /*InSessionFilter*/);
+
+	DECLARE_MULTICAST_DELEGATE_FiveParams(FCopySession, bool /*bSuccess*/, const IConcertServer& /*InServer*/, TSharedRef<IConcertServerSession> /*InLiveSession*/, const FString& /*NewSessionRoot*/, const FConcertSessionFilter& /*InSessionFilter*/);
+	DECLARE_MULTICAST_DELEGATE_SixParams(FExportSession, bool/*bSuccess*/, const IConcertServer& /*InServer*/, const FGuid& /*InSessionId*/, const FString& /*DestDir*/, const FConcertSessionFilter& /*InSessionFilter*/, bool /*bAnonymizeData*/);
+	DECLARE_MULTICAST_DELEGATE_SixParams(FRestoreSession, bool/*bSuccess*/, const IConcertServer& /*InServer*/, const FGuid& /*InArchivedSessionId*/, const FString& /*InLiveSessionRoot*/, const FConcertSessionInfo& /*InLiveSessionInfo*/, const FConcertSessionFilter& /*InSessionFilter*/);
+	
+	DECLARE_MULTICAST_DELEGATE_TwoParams(FOnLiveSessionRenamed, const IConcertServer& /*InServer*/, TSharedRef<IConcertServerSession> /*InLiveSession*/);
+	DECLARE_MULTICAST_DELEGATE_ThreeParams(FOnArchivedSessionRenamed, const IConcertServer& /*InServer*/, const FString& /*InArchivedSessionRoot*/, const FConcertSessionInfo& /*InArchivedSessionInfo*/);
+
+	CONCERT_API FOnLiveSessionCreated& OnLiveSessionCreated();
+	CONCERT_API FOnArchivedSessionCreated& OnArchivedSessionCreated();
+	CONCERT_API FOnLiveSessionDestroyed& OnLiveSessionDestroyed();
+	CONCERT_API FOnArchivedSessionDestroyed& OnArchivedSessionDestroyed();
+	CONCERT_API FArchiveSession_WithSession& ArchiveSession_WithSession();
+	CONCERT_API FArchiveSession_WithWorkingDir& ArchiveSession_WithWorkingDir();
+	CONCERT_API FCopySession& CopySession();
+	CONCERT_API FExportSession& ExportSession();
+	CONCERT_API FRestoreSession& RestoreSession();
+	CONCERT_API FOnLiveSessionRenamed& OnLiveSessionRenamed();
+	CONCERT_API FOnArchivedSessionRenamed& OnArchivedSessionRenamed();
+}
\ No newline at end of file
diff --git a/Engine/Plugins/Developer/Concert/ConcertSync/ConcertSyncServer/Source/ConcertSyncServer/Private/ConcertServerEventForwardingSink.h b/Engine/Plugins/Developer/Concert/ConcertSync/ConcertSyncServer/Source/ConcertSyncServer/Private/ConcertServerEventForwardingSink.h
new file mode 100644
--- /dev/null
+++ b/Engine/Plugins/Developer/Concert/ConcertSync/ConcertSyncServer/Source/ConcertSyncServer/Private/ConcertServerEventForwardingSink.h
@@ -0,0 +1,105 @@
+﻿// Copyright Epic Games, Inc. All Rights Reserved.
+
+#pragma once
+
+#include "CoreMinimal.h"
+#include "IConcertServerEventSink.h"
+#include "ConcertMessageData.h"
+#include "ConcertServerEvents.h"
+
+/**
+ * Uses CRTP to call IConcertServerEventSink::FooImpl and calls the equivalent event in ConcertServerEvents with the result.
+ * Example: TDerived must implement GetSessionsFromPathImpl.
+ */
+template<typename TDerived>
+class TConcertServerEventForwardingSink : public IConcertServerEventSink
+{
+	TDerived* This() { return static_cast<TDerived*>(this); }
+	const TDerived* This() const { return static_cast<TDerived*>(this); }
+public:
+
+	//~ Begin IConcertServerEventSink Interface
+	virtual void GetSessionsFromPath(const IConcertServer& InServer, const FString& InPath, TArray<FConcertSessionInfo>& OutSessionInfos, TArray<FDateTime>* OutSessionCreationTimes = nullptr) final
+	{
+		This()->GetSessionsFromPathImpl(InServer, InPath, OutSessionInfos, OutSessionCreationTimes);
+	}
+
+	virtual bool OnLiveSessionCreated(const IConcertServer& InServer, TSharedRef<IConcertServerSession> InLiveSession) final
+	{
+		const bool bSuccess = This()->OnLiveSessionCreatedImpl(InServer, InLiveSession);
+		ConcertServerEvents::OnLiveSessionCreated().Broadcast(bSuccess, InServer, InLiveSession);
+		return bSuccess;
+	}
+	
+	virtual void OnLiveSessionDestroyed(const IConcertServer& InServer, TSharedRef<IConcertServerSession> InLiveSession) final
+	{
+		This()->OnLiveSessionDestroyedImpl(InServer, InLiveSession);
+		ConcertServerEvents::OnLiveSessionDestroyed().Broadcast(InServer, InLiveSession);
+	}
+	
+	virtual bool OnArchivedSessionCreated(const IConcertServer& InServer, const FString& InArchivedSessionRoot, const FConcertSessionInfo& InArchivedSessionInfo) final
+	{
+		const bool bSuccess = This()->OnArchivedSessionCreatedImpl(InServer, InArchivedSessionRoot, InArchivedSessionInfo);
+		ConcertServerEvents::OnArchivedSessionCreated().Broadcast(bSuccess, InServer, InArchivedSessionRoot, InArchivedSessionInfo);
+		return bSuccess;
+	}
+	
+	virtual void OnArchivedSessionDestroyed(const IConcertServer& InServer, const FGuid& InArchivedSessionId) final
+	{
+		This()->OnArchivedSessionDestroyedImpl(InServer, InArchivedSessionId);
+		ConcertServerEvents::OnArchivedSessionDestroyed().Broadcast(InServer, InArchivedSessionId);
+	}
+	
+	virtual bool ArchiveSession(const IConcertServer& InServer, TSharedRef<IConcertServerSession> InLiveSession, const FString& InArchivedSessionRoot, const FConcertSessionInfo& InArchivedSessionInfo, const FConcertSessionFilter& InSessionFilter) final
+	{
+		const bool bSuccess = This()->ArchiveSessionImpl(InServer, InLiveSession, InArchivedSessionRoot, InArchivedSessionInfo, InSessionFilter);
+		ConcertServerEvents::ArchiveSession_WithSession().Broadcast(bSuccess, InServer, InLiveSession, InArchivedSessionRoot, InArchivedSessionInfo, InSessionFilter);
+		return bSuccess;
+	}
+	
+	virtual bool ArchiveSession(const IConcertServer& InServer, const FString& InLiveSessionWorkingDir, const FString& InArchivedSessionRoot, const FConcertSessionInfo& InArchivedSessionInfo, const FConcertSessionFilter& InSessionFilter) final
+	{
+		const bool bSuccess = This()->ArchiveSessionImpl(InServer, InLiveSessionWorkingDir, InArchivedSessionRoot, InArchivedSessionInfo, InSessionFilter);
+		ConcertServerEvents::ArchiveSession_WithWorkingDir().Broadcast(bSuccess, InServer, InLiveSessionWorkingDir, InArchivedSessionRoot, InArchivedSessionInfo, InSessionFilter);
+		return bSuccess;
+	}
+
+	virtual bool CopySession(const IConcertServer& InServer, TSharedRef<IConcertServerSession> InLiveSession, const FString& NewSessionRoot, const FConcertSessionFilter& InSessionFilter) final
+	{
+		const bool bSuccess = This()->CopySessionImpl(InServer, InLiveSession, NewSessionRoot, InSessionFilter);
+		ConcertServerEvents::CopySession().Broadcast(bSuccess, InServer, InLiveSession, NewSessionRoot, InSessionFilter);
+		return bSuccess;
+	}
+
+	virtual bool ExportSession(const IConcertServer& InServer, const FGuid& InSessionId, const FString& DestDir, const FConcertSessionFilter& InSessionFilter, bool bAnonymizeData) final
+	{
+		const bool bSuccess = This()->ExportSessionImpl(InServer, InSessionId, DestDir, InSessionFilter, bAnonymizeData);
+		ConcertServerEvents::ExportSession().Broadcast(bSuccess, InServer, InSessionId, DestDir, InSessionFilter, bAnonymizeData);
+		return bSuccess;
+	}
+
+	virtual bool RestoreSession(const IConcertServer& InServer, const FGuid& InArchivedSessionId, const FString& InLiveSessionRoot, const FConcertSessionInfo& InLiveSessionInfo, const FConcertSessionFilter& InSessionFilter) final
+	{
+		const bool bSuccess = This()->RestoreSessionImpl(InServer, InArchivedSessionId, InLiveSessionRoot, InLiveSessionInfo, InSessionFilter);
+		ConcertServerEvents::RestoreSession().Broadcast(bSuccess, InServer, InArchivedSessionId, InLiveSessionRoot, InLiveSessionInfo, InSessionFilter);
+		return bSuccess;
+	}
+
+	virtual bool GetSessionActivities(const IConcertServer& InServer, const FGuid& SessionId, int64 FromActivityId, int64 ActivityCount, TArray<FConcertSessionSerializedPayload>& OutActivities, TMap<FGuid, FConcertClientInfo>& OutEndpointClientInfoMap, bool bIncludeDetails) final
+	{
+		return This()->GetSessionActivitiesImpl(InServer, SessionId, FromActivityId, ActivityCount, OutActivities, OutEndpointClientInfoMap, bIncludeDetails);
+	}
+
+	virtual void OnLiveSessionRenamed(const IConcertServer& InServer, TSharedRef<IConcertServerSession> InLiveSession) final
+	{
+		This()->OnLiveSessionRenamedImpl(InServer, InLiveSession);
+		ConcertServerEvents::OnLiveSessionRenamed().Broadcast(InServer, InLiveSession);
+	}
+	
+	virtual void OnArchivedSessionRenamed(const IConcertServer& InServer, const FString& InArchivedSessionRoot, const FConcertSessionInfo& InArchivedSessionInfo) final
+	{
+		This()->OnArchivedSessionRenamedImpl(InServer, InArchivedSessionRoot, InArchivedSessionInfo);
+		ConcertServerEvents::OnArchivedSessionRenamed().Broadcast(InServer, InArchivedSessionRoot, InArchivedSessionInfo);
+	}
+	//~ End IConcertServerEventSink Interface
+};
\ No newline at end of file
diff --git a/Engine/Plugins/Developer/Concert/ConcertSync/ConcertSyncServer/Source/ConcertSyncServer/Private/ConcertSyncServer.cpp b/Engine/Plugins/Developer/Concert/ConcertSync/ConcertSyncServer/Source/ConcertSyncServer/Private/ConcertSyncServer.cpp
--- a/Engine/Plugins/Developer/Concert/ConcertSync/ConcertSyncServer/Source/ConcertSyncServer/Private/ConcertSyncServer.cpp
+++ b/Engine/Plugins/Developer/Concert/ConcertSync/ConcertSyncServer/Source/ConcertSyncServer/Private/ConcertSyncServer.cpp
@@ -261,7 +261,7 @@ void FConcertSyncServer::SetFileSharingService(TSharedPtr<IConcertFileSharingSer
 	FileSharingService = MoveTemp(InFileSharingService);
 }
 
-void FConcertSyncServer::GetSessionsFromPath(const IConcertServer& InServer, const FString& InPath, TArray<FConcertSessionInfo>& OutSessionInfos, TArray<FDateTime>* OutSessionCreationTimes)
+void FConcertSyncServer::GetSessionsFromPathImpl(const IConcertServer& InServer, const FString& InPath, TArray<FConcertSessionInfo>& OutSessionInfos, TArray<FDateTime>* OutSessionCreationTimes)
 {
 	IFileManager::Get().IterateDirectory(*InPath, [&OutSessionInfos, &OutSessionCreationTimes](const TCHAR* FilenameOrDirectory, bool bIsDirectory) -> bool
 	{
@@ -290,35 +290,35 @@ void FConcertSyncServer::GetSessionsFromPath(const IConcertServer& InServer, con
 	});
 }
 
-bool FConcertSyncServer::OnLiveSessionCreated(const IConcertServer& InServer, TSharedRef<IConcertServerSession> InSession)
+bool FConcertSyncServer::OnLiveSessionCreatedImpl(const IConcertServer& InServer, TSharedRef<IConcertServerSession> InSession)
 {
 	ConcertSyncServerUtils::WriteSessionInfoToDirectory(InSession->GetSessionWorkingDirectory(), InSession->GetSessionInfo());
 	return CreateLiveSession(InSession);
 }
 
-void FConcertSyncServer::OnLiveSessionDestroyed(const IConcertServer& InServer, TSharedRef<IConcertServerSession> InSession)
+void FConcertSyncServer::OnLiveSessionDestroyedImpl(const IConcertServer& InServer, TSharedRef<IConcertServerSession> InSession)
 {
 	DestroyLiveSession(InSession);
 }
 
-bool FConcertSyncServer::OnArchivedSessionCreated(const IConcertServer& InServer, const FString& InArchivedSessionRoot, const FConcertSessionInfo& InArchivedSessionInfo)
+bool FConcertSyncServer::OnArchivedSessionCreatedImpl(const IConcertServer& InServer, const FString& InArchivedSessionRoot, const FConcertSessionInfo& InArchivedSessionInfo)
 {
 	ConcertSyncServerUtils::WriteSessionInfoToDirectory(InArchivedSessionRoot, InArchivedSessionInfo);
 	return CreateArchivedSession(InArchivedSessionRoot, InArchivedSessionInfo);
 }
 
-void FConcertSyncServer::OnArchivedSessionDestroyed(const IConcertServer& InServer, const FGuid& InArchivedSessionId)
+void FConcertSyncServer::OnArchivedSessionDestroyedImpl(const IConcertServer& InServer, const FGuid& InArchivedSessionId)
 {
 	DestroyArchivedSession(InArchivedSessionId);
 }
 
-bool FConcertSyncServer::ArchiveSession(const IConcertServer& InServer, TSharedRef<IConcertServerSession> InLiveSession, const FString& InArchivedSessionRoot, const FConcertSessionInfo& InArchivedSessionInfo, const FConcertSessionFilter& InSessionFilter)
+bool FConcertSyncServer::ArchiveSessionImpl(const IConcertServer& InServer, TSharedRef<IConcertServerSession> InLiveSession, const FString& InArchivedSessionRoot, const FConcertSessionInfo& InArchivedSessionInfo, const FConcertSessionFilter& InSessionFilter)
 {
 	SCOPED_CONCERT_TRACE(ConcertServer_ArchiveSession);
-	return CopySession(InServer, InLiveSession, InArchivedSessionRoot, InSessionFilter);
+	return CopySessionImpl(InServer, InLiveSession, InArchivedSessionRoot, InSessionFilter);
 }
 
-bool FConcertSyncServer::ArchiveSession(const IConcertServer& InServer, const FString& InLiveSessionWorkingDir, const FString& InArchivedSessionRoot, const FConcertSessionInfo& InArchivedSessionInfo, const FConcertSessionFilter& InSessionFilter)
+bool FConcertSyncServer::ArchiveSessionImpl(const IConcertServer& InServer, const FString& InLiveSessionWorkingDir, const FString& InArchivedSessionRoot, const FConcertSessionInfo& InArchivedSessionInfo, const FConcertSessionFilter& InSessionFilter)
 {
 	FConcertSyncSessionDatabase LiveSessionDatabase;
 	LiveSessionDatabase.Open(InLiveSessionWorkingDir);
@@ -330,7 +330,7 @@ bool FConcertSyncServer::ArchiveSession(const IConcertServer& InServer, const FS
 	return RetVal;
 }
 
-bool FConcertSyncServer::CopySession(const IConcertServer& InServer, TSharedRef<IConcertServerSession> InLiveSession, const FString& NewSessionRoot, const FConcertSessionFilter& InSessionFilter)
+bool FConcertSyncServer::CopySessionImpl(const IConcertServer& InServer, TSharedRef<IConcertServerSession> InLiveSession, const FString& NewSessionRoot, const FConcertSessionFilter& InSessionFilter)
 {
 	SCOPED_CONCERT_TRACE(ConcertServer_CopySession);
 	if (TSharedPtr<FConcertSyncServerLiveSession> LiveSession = LiveSessions.FindRef(InLiveSession->GetId()))
@@ -340,7 +340,7 @@ bool FConcertSyncServer::CopySession(const IConcertServer& InServer, TSharedRef<
 	return false;
 }
 
-bool FConcertSyncServer::ExportSession(const IConcertServer& InServer, const FGuid& InSessionId, const FString& DestDir, const FConcertSessionFilter& InSessionFilter, bool bAnonymizeData)
+bool FConcertSyncServer::ExportSessionImpl(const IConcertServer& InServer, const FGuid& InSessionId, const FString& DestDir, const FConcertSessionFilter& InSessionFilter, bool bAnonymizeData)
 {
 	SCOPED_CONCERT_TRACE(ConcertServer_ExportSession);
 	if (TSharedPtr<FConcertSyncServerLiveSession> LiveSession = LiveSessions.FindRef(InSessionId)) // If the session is live.
@@ -357,7 +357,7 @@ bool FConcertSyncServer::ExportSession(const IConcertServer& InServer, const FGu
 	return false; // Session not found.
 }
 
-bool FConcertSyncServer::RestoreSession(const IConcertServer& InServer, const FGuid& InArchivedSessionId, const FString& InLiveSessionRoot, const FConcertSessionInfo& InLiveSessionInfo, const FConcertSessionFilter& InSessionFilter)
+bool FConcertSyncServer::RestoreSessionImpl(const IConcertServer& InServer, const FGuid& InArchivedSessionId, const FString& InLiveSessionRoot, const FConcertSessionInfo& InLiveSessionInfo, const FConcertSessionFilter& InSessionFilter)
 {
 	SCOPED_CONCERT_TRACE(ConcertServer_RestoreSession);
 
@@ -368,22 +368,22 @@ bool FConcertSyncServer::RestoreSession(const IConcertServer& InServer, const FG
 	return false;
 }
 
-bool FConcertSyncServer::GetSessionActivities(const IConcertServer& InServer, const FGuid& SessionId, int64 FromActivityId, int64 ActivityCount, TArray<FConcertSessionSerializedPayload>& Activities, TMap<FGuid, FConcertClientInfo>& OutEndpointClientInfoMap, bool bIncludeDetails)
+bool FConcertSyncServer::GetSessionActivitiesImpl(const IConcertServer& InServer, const FGuid& SessionId, int64 FromActivityId, int64 ActivityCount, TArray<FConcertSessionSerializedPayload>& Activities, TMap<FGuid, FConcertClientInfo>& OutEndpointClientInfoMap, bool bIncludeDetails)
 {
 	SCOPED_CONCERT_TRACE(ConcertServer_GetSessionActivities);
 	if (TSharedPtr<FConcertSyncServerLiveSession> LiveSession = LiveSessions.FindRef(SessionId))
 	{
-		return GetSessionActivities(LiveSession->GetSessionDatabase(), FromActivityId, ActivityCount, Activities, OutEndpointClientInfoMap, bIncludeDetails);
+		return GetSessionActivitiesInternal(LiveSession->GetSessionDatabase(), FromActivityId, ActivityCount, Activities, OutEndpointClientInfoMap, bIncludeDetails);
 	}
 	else if (TSharedPtr<FConcertSyncServerArchivedSession> ArchivedSession = ArchivedSessions.FindRef(SessionId))
 	{
-		return GetSessionActivities(ArchivedSession->GetSessionDatabase(), FromActivityId, ActivityCount, Activities, OutEndpointClientInfoMap, bIncludeDetails);
+		return GetSessionActivitiesInternal(ArchivedSession->GetSessionDatabase(), FromActivityId, ActivityCount, Activities, OutEndpointClientInfoMap, bIncludeDetails);
 	}
 
 	return false; // Not Found.
 }
 
-bool FConcertSyncServer::GetSessionActivities(const FConcertSyncSessionDatabase& Database, int64 FromActivityId, int64 ActivityCount, TArray<FConcertSessionSerializedPayload>& OutActivities, TMap<FGuid, FConcertClientInfo>& OutEndpointClientInfoMap, bool bIncludeDetails)
+bool FConcertSyncServer::GetSessionActivitiesInternal(const FConcertSyncSessionDatabase& Database, int64 FromActivityId, int64 ActivityCount, TArray<FConcertSessionSerializedPayload>& OutActivities, TMap<FGuid, FConcertClientInfo>& OutEndpointClientInfoMap, bool bIncludeDetails)
 {
 	int64 MaxActivityId;
 	Database.GetActivityMaxId(MaxActivityId);
@@ -463,12 +463,12 @@ bool FConcertSyncServer::GetSessionActivities(const FConcertSyncSessionDatabase&
 	return true;
 }
 
-void FConcertSyncServer::OnLiveSessionRenamed(const IConcertServer& InServer,TSharedRef<IConcertServerSession> InLiveSession)
+void FConcertSyncServer::OnLiveSessionRenamedImpl(const IConcertServer& InServer,TSharedRef<IConcertServerSession> InLiveSession)
 {
 	ConcertSyncServerUtils::WriteSessionInfoToDirectory(InLiveSession->GetSessionWorkingDirectory(), InLiveSession->GetSessionInfo());
 }
 
-void FConcertSyncServer::OnArchivedSessionRenamed(const IConcertServer& InServer, const FString& InArchivedSessionRoot, const FConcertSessionInfo& InArchivedSessionInfo)
+void FConcertSyncServer::OnArchivedSessionRenamedImpl(const IConcertServer& InServer, const FString& InArchivedSessionRoot, const FConcertSessionInfo& InArchivedSessionInfo)
 {
 	ConcertSyncServerUtils::WriteSessionInfoToDirectory(InArchivedSessionRoot, InArchivedSessionInfo);
 }
diff --git a/Engine/Plugins/Developer/Concert/ConcertSync/ConcertSyncServer/Source/ConcertSyncServer/Private/ConcertSyncServer.h b/Engine/Plugins/Developer/Concert/ConcertSync/ConcertSyncServer/Source/ConcertSyncServer/Private/ConcertSyncServer.h
--- a/Engine/Plugins/Developer/Concert/ConcertSync/ConcertSyncServer/Source/ConcertSyncServer/Private/ConcertSyncServer.h
+++ b/Engine/Plugins/Developer/Concert/ConcertSync/ConcertSyncServer/Source/ConcertSyncServer/Private/ConcertSyncServer.h
@@ -2,6 +2,7 @@
 
 #pragma once
 
+#include "ConcertServerEventForwardingSink.h"
 #include "IConcertSyncServer.h"
 #include "IConcertServerEventSink.h"
 
@@ -17,7 +18,7 @@ struct FConcertSessionFilter;
 /**
  * Implementation for a Concert Sync Server.
  */
-class FConcertSyncServer : public IConcertSyncServer, public IConcertServerEventSink
+class FConcertSyncServer : public IConcertSyncServer, public TConcertServerEventForwardingSink<FConcertSyncServer>
 {
 public:
 	FConcertSyncServer(const FString& InRole, const FConcertSessionFilter& InAutoArchiveSessionFilter);
@@ -30,19 +31,19 @@ public:
 	virtual void SetFileSharingService(TSharedPtr<IConcertFileSharingService> InFileSharingService) override;
 
 	//~ IConcertServerEventSink interface
-	virtual void GetSessionsFromPath(const IConcertServer& InServer, const FString& InPath, TArray<FConcertSessionInfo>& OutSessionInfos, TArray<FDateTime>* OutSessionCreationTimes = nullptr) override;
-	virtual bool OnLiveSessionCreated(const IConcertServer& InServer, TSharedRef<IConcertServerSession> InLiveSession) override;
-	virtual void OnLiveSessionDestroyed(const IConcertServer& InServer, TSharedRef<IConcertServerSession> InLiveSession) override;
-	virtual bool OnArchivedSessionCreated(const IConcertServer& InServer, const FString& InArchivedSessionRoot, const FConcertSessionInfo& InArchivedSessionInfo) override;
-	virtual void OnArchivedSessionDestroyed(const IConcertServer& InServer, const FGuid& InArchivedSessionId) override;
-	virtual bool ArchiveSession(const IConcertServer& InServer, TSharedRef<IConcertServerSession> InLiveSession, const FString& InArchivedSessionRoot, const FConcertSessionInfo& InArchivedSessionInfo, const FConcertSessionFilter& InSessionFilter) override;
-	virtual bool ArchiveSession(const IConcertServer& InServer, const FString& InLiveSessionWorkingDir, const FString& InArchivedSessionRoot, const FConcertSessionInfo& InArchivedSessionInfo, const FConcertSessionFilter& InSessionFilter) override;
-	virtual bool CopySession(const IConcertServer& InServer, TSharedRef<IConcertServerSession> InLiveSession, const FString& NewSessionRoot, const FConcertSessionFilter& InSessionFilter) override;
-	virtual bool ExportSession(const IConcertServer& InServer, const FGuid& InSessionId, const FString& DestDir, const FConcertSessionFilter& InSessionFilter, bool bAnonymizeData) override;
-	virtual bool RestoreSession(const IConcertServer& InServer, const FGuid& InArchivedSessionId, const FString& InLiveSessionRoot, const FConcertSessionInfo& InLiveSessionInfo, const FConcertSessionFilter& InSessionFilter) override;
-	virtual bool GetSessionActivities(const IConcertServer& InServer, const FGuid& SessionId, int64 FromActivityId, int64 ActivityCount, TArray<FConcertSessionSerializedPayload>& OutActivities, TMap<FGuid, FConcertClientInfo>& OutEndpointClientInfoMap, bool bIncludeDetails) override;
-	virtual void OnLiveSessionRenamed(const IConcertServer& InServer, TSharedRef<IConcertServerSession> InLiveSession) override;
-	virtual void OnArchivedSessionRenamed(const IConcertServer& InServer, const FString& InArchivedSessionRoot, const FConcertSessionInfo& InArchivedSessionInfo) override;
+	void GetSessionsFromPathImpl(const IConcertServer& InServer, const FString& InPath, TArray<FConcertSessionInfo>& OutSessionInfos, TArray<FDateTime>* OutSessionCreationTimes = nullptr);
+	bool OnLiveSessionCreatedImpl(const IConcertServer& InServer, TSharedRef<IConcertServerSession> InLiveSession);
+	void OnLiveSessionDestroyedImpl(const IConcertServer& InServer, TSharedRef<IConcertServerSession> InLiveSession);
+	bool OnArchivedSessionCreatedImpl(const IConcertServer& InServer, const FString& InArchivedSessionRoot, const FConcertSessionInfo& InArchivedSessionInfo);
+	void OnArchivedSessionDestroyedImpl(const IConcertServer& InServer, const FGuid& InArchivedSessionId);
+	bool ArchiveSessionImpl(const IConcertServer& InServer, TSharedRef<IConcertServerSession> InLiveSession, const FString& InArchivedSessionRoot, const FConcertSessionInfo& InArchivedSessionInfo, const FConcertSessionFilter& InSessionFilter);
+	bool ArchiveSessionImpl(const IConcertServer& InServer, const FString& InLiveSessionWorkingDir, const FString& InArchivedSessionRoot, const FConcertSessionInfo& InArchivedSessionInfo, const FConcertSessionFilter& InSessionFilter);
+	bool CopySessionImpl(const IConcertServer& InServer, TSharedRef<IConcertServerSession> InLiveSession, const FString& NewSessionRoot, const FConcertSessionFilter& InSessionFilter);
+	bool ExportSessionImpl(const IConcertServer& InServer, const FGuid& InSessionId, const FString& DestDir, const FConcertSessionFilter& InSessionFilter, bool bAnonymizeData);
+	bool RestoreSessionImpl(const IConcertServer& InServer, const FGuid& InArchivedSessionId, const FString& InLiveSessionRoot, const FConcertSessionInfo& InLiveSessionInfo, const FConcertSessionFilter& InSessionFilter);
+	bool GetSessionActivitiesImpl(const IConcertServer& InServer, const FGuid& SessionId, int64 FromActivityId, int64 ActivityCount, TArray<FConcertSessionSerializedPayload>& OutActivities, TMap<FGuid, FConcertClientInfo>& OutEndpointClientInfoMap, bool bIncludeDetails);
+	void OnLiveSessionRenamedImpl(const IConcertServer& InServer, TSharedRef<IConcertServerSession> InLiveSession);
+	void OnArchivedSessionRenamedImpl(const IConcertServer& InServer, const FString& InArchivedSessionRoot, const FConcertSessionInfo& InArchivedSessionInfo);
 
 private:
 	void CreateWorkspace(const TSharedRef<FConcertSyncServerLiveSession>& InLiveSession);
@@ -57,7 +58,7 @@ private:
 	bool CreateArchivedSession(const FString& InArchivedSessionRoot, const FConcertSessionInfo& InArchivedSessionInfo);
 	void DestroyArchivedSession(const FGuid& InArchivedSessionId);
 
-	bool GetSessionActivities(const FConcertSyncSessionDatabase& Database, int64 FromActivityId, int64 ActivityCount, TArray<FConcertSessionSerializedPayload>& OutActivities, TMap<FGuid, FConcertClientInfo>& OutEndpointClientInfoMap, bool bIncludeDetails);
+	bool GetSessionActivitiesInternal(const FConcertSyncSessionDatabase& Database, int64 FromActivityId, int64 ActivityCount, TArray<FConcertSessionSerializedPayload>& OutActivities, TMap<FGuid, FConcertClientInfo>& OutEndpointClientInfoMap, bool bIncludeDetails);
 
 	/** Server for Concert */
 	IConcertServerRef ConcertServer;
