Concert: delay sending Sequencer state sync events to the client until its workspace is synced and finalized

This addresses an issue particularly with nDisplay nodes where the client can receive
the Sequencer state sync event before it has finished processing transactions from
the activity stream. If a sequence open on another client was created or saved by a
transaction in that stream, the nDisplay node may not be ready to open or sync a
Sequencer/player for it until its workspace is finalized.

In addition, a Sequencer/player is opened for any sequence not already open on the
client when the Sequencer state sync event is received. Since the client may have
been launched some time after other clients, it would not have received Sequencer
open events that occurred before it was launched.

This also includes some minor cleanup to event handling and delegate registration and
unregistration, as well as a few small corrections to some docstrings.

diff --git a/Engine/Plugins/Developer/Concert/ConcertSync/ConcertSyncClient/Source/ConcertSyncClient/Private/ConcertClientSequencerManager.cpp b/Engine/Plugins/Developer/Concert/ConcertSync/ConcertSyncClient/Source/ConcertSyncClient/Private/ConcertClientSequencerManager.cpp
--- a/Engine/Plugins/Developer/Concert/ConcertSync/ConcertSyncClient/Source/ConcertSyncClient/Private/ConcertClientSequencerManager.cpp
+++ b/Engine/Plugins/Developer/Concert/ConcertSync/ConcertSyncClient/Source/ConcertSyncClient/Private/ConcertClientSequencerManager.cpp
@@ -144,7 +144,7 @@ FConcertClientSequencerManager::GatherRootSequencersByState(const FString& InSeq
 		TSharedPtr<ISequencer> Sequencer = Entry.WeakSequencer.Pin();
 		UMovieSceneSequence*   Sequence = Sequencer.IsValid() ? Sequencer->GetRootMovieSceneSequence() : nullptr;
 
-		if (Sequence && (Sequence->GetPathName() == InSequenceObjectPath || CVarEnableUnrelatedTimelineSync.GetValueOnAnyThread() > 0))
+		if (Sequence && (Sequence->GetPathName() == InSequenceObjectPath || IsUnrelatedSequencerTimelineSyncEnabled()))
 		{
 			OutSequencers.Add(&Entry);
 		}
@@ -285,7 +285,16 @@ void FConcertClientSequencerManager::OnSyncEvent(const FConcertSessionContext& I
 	{
 		FConcertSequencerState& SequencerState = SequencerStates.FindOrAdd(*State.SequenceObjectPath);
 		SequencerState = State;
-		for (FOpenSequencerData* OpenSequencer : GatherRootSequencersByState(SequencerState.SequenceObjectPath))
+
+		const TArray<FOpenSequencerData*, TInlineAllocator<1>> OpenSequencersForObject = GatherRootSequencersByState(SequencerState.SequenceObjectPath);
+
+		// Keep track of whether a Sequencer is open for this particular
+		// sequence, since the array of open Sequencers may contain unrelated
+		// Sequencers.
+		bool bFoundSequencerForObject = false;
+
+		UE_LOG(LogConcertSequencerSync, VeryVerbose, TEXT("OnSyncEvent: Syncing %d Sequencers for sequence %s"), OpenSequencersForObject.Num(), *SequencerState.SequenceObjectPath);
+		for (FOpenSequencerData* OpenSequencer : OpenSequencersForObject)
 		{
 			TSharedPtr<ISequencer> Sequencer = OpenSequencer->WeakSequencer.Pin();
 			if (Sequencer.IsValid() && IsSequencerPlaybackSyncEnabled())
@@ -293,8 +302,21 @@ void FConcertClientSequencerManager::OnSyncEvent(const FConcertSessionContext& I
 				Sequencer->SetGlobalTime(SequencerState.Time.ConvertTo(Sequencer->GetRootTickResolution()));
 				Sequencer->SetPlaybackStatus((EMovieScenePlayerStatus::Type)SequencerState.PlayerStatus);
 				Sequencer->SetPlaybackSpeed(SequencerState.PlaybackSpeed);
+
+				const UMovieSceneSequence* Sequence = Sequencer->GetRootMovieSceneSequence();
+				if (Sequence && (Sequence->GetPathName() == SequencerState.SequenceObjectPath))
+				{
+					bFoundSequencerForObject = true;
+				}
 			}
 		}
+
+		if (!bFoundSequencerForObject)
+		{
+			UE_LOG(LogConcertSequencerSync, VeryVerbose, TEXT("OnSyncEvent: No existing Sequencer with sequence %s open. Will open and sync a new one at end of frame."), *SequencerState.SequenceObjectPath);
+			PendingSequenceOpenEvents.Add(SequencerState.SequenceObjectPath);
+			PendingSequencerEvents.Add(SequencerState);
+		}
 	}
 }
 
diff --git a/Engine/Plugins/Developer/Concert/ConcertSync/ConcertSyncClient/Source/ConcertSyncClient/Private/ConcertClientSequencerManager.h b/Engine/Plugins/Developer/Concert/ConcertSync/ConcertSyncClient/Source/ConcertSyncClient/Private/ConcertClientSequencerManager.h
--- a/Engine/Plugins/Developer/Concert/ConcertSync/ConcertSyncClient/Source/ConcertSyncClient/Private/ConcertClientSequencerManager.h
+++ b/Engine/Plugins/Developer/Concert/ConcertSync/ConcertSyncClient/Source/ConcertSyncClient/Private/ConcertClientSequencerManager.h
@@ -176,10 +176,10 @@ private:
 	void OnOpenEvent(const FConcertSessionContext&, const FConcertSequencerOpenEvent& InEvent);
 
 	/**
-	 * Called on receipt of an external close event from the server
+	 * Called on receipt of an external state sync event from the server
 	 *
 	 * @param InEventContext              The context for the current client session
-	 * @param InEvent                     The sequencer close event received from the server
+	 * @param InEvent                     The sequencer state sync event received from the server
 	 */
 	void OnSyncEvent(const FConcertSessionContext& InEventContext, const FConcertSequencerStateSyncEvent& InEvent);
 
diff --git a/Engine/Plugins/Developer/Concert/ConcertSync/ConcertSyncClient/Source/ConcertSyncClient/Private/ConcertClientWorkspace.cpp b/Engine/Plugins/Developer/Concert/ConcertSync/ConcertSyncClient/Source/ConcertSyncClient/Private/ConcertClientWorkspace.cpp
--- a/Engine/Plugins/Developer/Concert/ConcertSync/ConcertSyncClient/Source/ConcertSyncClient/Private/ConcertClientWorkspace.cpp
+++ b/Engine/Plugins/Developer/Concert/ConcertSync/ConcertSyncClient/Source/ConcertSyncClient/Private/ConcertClientWorkspace.cpp
@@ -19,6 +19,7 @@
 #include "ConcertSyncClientUtil.h"
 #include "ConcertLogGlobal.h"
 #include "ConcertWorkspaceData.h"
+#include "ConcertWorkspaceMessages.h"
 #include "ConcertClientDataStore.h"
 #include "ConcertClientLiveTransactionAuthors.h"
 
@@ -847,6 +848,9 @@ void FConcertClientWorkspace::OnEndFrame()
 		bHasSyncedWorkspace = true;
 		FConcertSlowTaskStackWorkaround::Get().PopTask(MoveTemp(InitialSyncSlowTask));
 		OnFinalizeWorkspaceSyncCompletedDelegate.Broadcast();
+
+		// Notify the server that we've finalized our workspace.
+		LiveSession->GetSession().SendCustomEvent(FConcertWorkspaceSyncAndFinalizeCompletedEvent(), LiveSession->GetSession().GetSessionServerEndpointId(), EConcertMessageFlags::ReliableOrdered);
 	}
 
 	if (bHasSyncedWorkspace && CanProcessPendingPackages() && !ConcertSyncClientUtil::UserIsEditing())
diff --git a/Engine/Plugins/Developer/Concert/ConcertSync/ConcertSyncCore/Source/ConcertSyncCore/Public/ConcertWorkspaceMessages.h b/Engine/Plugins/Developer/Concert/ConcertSync/ConcertSyncCore/Source/ConcertSyncCore/Public/ConcertWorkspaceMessages.h
--- a/Engine/Plugins/Developer/Concert/ConcertSync/ConcertSyncCore/Source/ConcertSyncCore/Public/ConcertWorkspaceMessages.h
+++ b/Engine/Plugins/Developer/Concert/ConcertSync/ConcertSyncCore/Source/ConcertSyncCore/Public/ConcertWorkspaceMessages.h
@@ -66,6 +66,13 @@ struct FConcertWorkspaceSyncCompletedEvent
 	GENERATED_BODY()
 };
 
+/** An event emitted by a client after its workspace has been completely synced and finalized. All transactions are posted and packages have been loaded. */
+USTRUCT()
+struct FConcertWorkspaceSyncAndFinalizeCompletedEvent
+{
+	GENERATED_BODY()
+};
+
 /** Request to sync an event that was partially synced on the client but for which the full data is required for inspection. FConcertSyncEventResponse is the corresponding response. */
 USTRUCT()
 struct FConcertSyncEventRequest
diff --git a/Engine/Plugins/Developer/Concert/ConcertSync/ConcertSyncServer/Source/ConcertSyncServer/Private/ConcertServerSequencerManager.cpp b/Engine/Plugins/Developer/Concert/ConcertSync/ConcertSyncServer/Source/ConcertSyncServer/Private/ConcertServerSequencerManager.cpp
--- a/Engine/Plugins/Developer/Concert/ConcertSync/ConcertSyncServer/Source/ConcertSyncServer/Private/ConcertServerSequencerManager.cpp
+++ b/Engine/Plugins/Developer/Concert/ConcertSync/ConcertSyncServer/Source/ConcertSyncServer/Private/ConcertServerSequencerManager.cpp
@@ -3,6 +3,7 @@
 #include "ConcertServerSequencerManager.h"
 #include "IConcertSession.h"
 #include "ConcertSyncServerLiveSession.h"
+#include "ConcertWorkspaceMessages.h"
 
 FConcertServerSequencerManager::FConcertServerSequencerManager(const TSharedRef<FConcertSyncServerLiveSession>& InLiveSession)
 {
@@ -26,6 +27,7 @@ void FConcertServerSequencerManager::BindSession(const TSharedRef<FConcertSyncSe
 	LiveSession->GetSession().RegisterCustomEventHandler<FConcertSequencerStateEvent>(this, &FConcertServerSequencerManager::HandleSequencerStateEvent);
 	LiveSession->GetSession().RegisterCustomEventHandler<FConcertSequencerOpenEvent>(this, &FConcertServerSequencerManager::HandleSequencerOpenEvent);
 	LiveSession->GetSession().RegisterCustomEventHandler<FConcertSequencerTimeAdjustmentEvent>(this, &FConcertServerSequencerManager::HandleSequencerTimeAdjustmentEvent);
+	LiveSession->GetSession().RegisterCustomEventHandler<FConcertWorkspaceSyncAndFinalizeCompletedEvent>(this, &FConcertServerSequencerManager::HandleWorkspaceSyncAndFinalizeCompletedEvent);
 }
 
 void FConcertServerSequencerManager::UnbindSession()
@@ -33,9 +35,12 @@ void FConcertServerSequencerManager::UnbindSession()
 	if (LiveSession)
 	{
 		LiveSession->GetSession().OnSessionClientChanged().RemoveAll(this);
-		LiveSession->GetSession().UnregisterCustomEventHandler<FConcertSequencerOpenEvent>(this);
 		LiveSession->GetSession().UnregisterCustomEventHandler<FConcertSequencerCloseEvent>(this);
 		LiveSession->GetSession().UnregisterCustomEventHandler<FConcertSequencerStateEvent>(this);
+		LiveSession->GetSession().UnregisterCustomEventHandler<FConcertSequencerOpenEvent>(this);
+		LiveSession->GetSession().UnregisterCustomEventHandler<FConcertSequencerTimeAdjustmentEvent>(this);
+		LiveSession->GetSession().UnregisterCustomEventHandler<FConcertWorkspaceSyncAndFinalizeCompletedEvent>(this);
+
 		LiveSession.Reset();
 	}
 }
@@ -100,6 +105,17 @@ void FConcertServerSequencerManager::HandleSequencerCloseEvent(const FConcertSes
 	}
 }
 
+void FConcertServerSequencerManager::HandleWorkspaceSyncAndFinalizeCompletedEvent(const FConcertSessionContext& InEventContext, const FConcertWorkspaceSyncAndFinalizeCompletedEvent& InEvent)
+{
+	FConcertSequencerStateSyncEvent SequencerStateSyncEvent;
+	for (const auto& Pair : SequencerStates)
+	{
+		SequencerStateSyncEvent.SequencerStates.Add(Pair.Value.State);
+	}
+
+	LiveSession->GetSession().SendCustomEvent(SequencerStateSyncEvent, InEventContext.SourceEndpointId, EConcertMessageFlags::ReliableOrdered);
+}
+
 void FConcertServerSequencerManager::HandleSessionClientChanged(IConcertServerSession& InSession, EConcertClientStatus InClientStatus, const FConcertSessionClientInfo& InClientInfo)
 {
 	check(&InSession == &LiveSession->GetSession());
@@ -122,14 +138,9 @@ void FConcertServerSequencerManager::HandleSessionClientChanged(IConcertServerSe
 			}
 		}
 	}
-	// Send the current Sequencers states to the newly connected client
-	else
-	{
-		FConcertSequencerStateSyncEvent SyncEvent;
-		for (const auto& Pair : SequencerStates)
-		{
-			SyncEvent.SequencerStates.Add(Pair.Value.State);
-		}
-		LiveSession->GetSession().SendCustomEvent(SyncEvent, InClientInfo.ClientEndpointId, EConcertMessageFlags::ReliableOrdered);
-	}
+
+	// Newly connected clients won't be sent the Sequencer state sync event
+	// until they have synced and finalized their workspace since an open
+	// sequence could have been created by a transaction in the activity
+	// stream.
 }
diff --git a/Engine/Plugins/Developer/Concert/ConcertSync/ConcertSyncServer/Source/ConcertSyncServer/Private/ConcertServerSequencerManager.h b/Engine/Plugins/Developer/Concert/ConcertSync/ConcertSyncServer/Source/ConcertSyncServer/Private/ConcertServerSequencerManager.h
--- a/Engine/Plugins/Developer/Concert/ConcertSync/ConcertSyncServer/Source/ConcertSyncServer/Private/ConcertServerSequencerManager.h
+++ b/Engine/Plugins/Developer/Concert/ConcertSync/ConcertSyncServer/Source/ConcertSyncServer/Private/ConcertServerSequencerManager.h
@@ -5,6 +5,7 @@
 #include "CoreMinimal.h"
 #include "ConcertMessages.h"
 #include "ConcertSequencerMessages.h"
+#include "ConcertWorkspaceMessages.h"
 
 struct FConcertSessionContext;
 class IConcertServerSession;
@@ -32,9 +33,12 @@ private:
 	/** Handler for the sequencer close event. */
 	void HandleSequencerCloseEvent(const FConcertSessionContext& InEventContext, const FConcertSequencerCloseEvent& InEvent);
 
-	/** Handler for the sequencer close event. */
+	/** Handler for the sequencer time adjustment event. */
 	void HandleSequencerTimeAdjustmentEvent(const FConcertSessionContext& InEventContext, const FConcertSequencerTimeAdjustmentEvent& InEvent);
 
+	/** Handler for the workspace sync and finalize completed event. */
+	void HandleWorkspaceSyncAndFinalizeCompletedEvent(const FConcertSessionContext& InEventContext, const FConcertWorkspaceSyncAndFinalizeCompletedEvent& InEvent);
+
 	/** Handler for the session clients changed event. */
 	void HandleSessionClientChanged(IConcertServerSession& InSession, EConcertClientStatus InClientStatus, const FConcertSessionClientInfo& InClientInfo);
 
diff --git a/Engine/Plugins/Runtime/nDisplay/Source/DisplayClusterMessageInterceptor/Private/DisplayClusterMessageInterceptionModule.cpp b/Engine/Plugins/Runtime/nDisplay/Source/DisplayClusterMessageInterceptor/Private/DisplayClusterMessageInterceptionModule.cpp
--- a/Engine/Plugins/Runtime/nDisplay/Source/DisplayClusterMessageInterceptor/Private/DisplayClusterMessageInterceptionModule.cpp
+++ b/Engine/Plugins/Runtime/nDisplay/Source/DisplayClusterMessageInterceptor/Private/DisplayClusterMessageInterceptionModule.cpp
@@ -50,9 +50,12 @@ class FDisplayClusterMessageInterceptionModule : public IModuleInterface
 public:
 	virtual void StartupModule() override
 	{
-		// Register for Cluster StartSession callback so everything is setup before launching interception
-		IDisplayCluster::Get().GetCallbacks().OnDisplayClusterStartSession().AddRaw(this, &FDisplayClusterMessageInterceptionModule::OnDisplayClusterStartSession);
-		IDisplayCluster::Get().GetCallbacks().OnDisplayClusterStartScene().AddRaw(this,&FDisplayClusterMessageInterceptionModule::OnNewSceneEvent);
+		if (IDisplayCluster::IsAvailable())
+		{
+			// Register for Cluster StartSession callback so everything is setup before launching interception
+			IDisplayCluster::Get().GetCallbacks().OnDisplayClusterStartSession().AddRaw(this, &FDisplayClusterMessageInterceptionModule::OnDisplayClusterStartSession);
+			IDisplayCluster::Get().GetCallbacks().OnDisplayClusterStartScene().AddRaw(this, &FDisplayClusterMessageInterceptionModule::OnNewSceneEvent);
+		}
 
 		// Setup console command to start/stop interception
 		StartMessageSyncCommand = MakeUnique<FAutoConsoleCommand>(
@@ -96,6 +99,7 @@ public:
 				IConcertClientRef ConcertClient = ConcertSyncClient->GetConcertClient();
 				ConcertClient->OnSessionStartup().RemoveAll(this);
 				ConcertClient->OnSessionShutdown().RemoveAll(this);
+				ConcertClient->OnSessionConnectionChanged().RemoveAll(this);
 			}
 		}
 #endif
@@ -112,6 +116,7 @@ public:
 
 			// Unregister cluster session events
 			IDisplayCluster::Get().GetCallbacks().OnDisplayClusterStartSession().RemoveAll(this);
+			IDisplayCluster::Get().GetCallbacks().OnDisplayClusterStartScene().RemoveAll(this);
 			IDisplayCluster::Get().GetCallbacks().OnDisplayClusterEndSession().RemoveAll(this);
 			IDisplayCluster::Get().GetCallbacks().OnDisplayClusterPreTick().RemoveAll(this);
 		}
