Add events for when a package starts being sent and has been received

diff --git a/Engine/Plugins/Developer/Concert/ConcertMain/Source/Concert/Private/ConcertMessageData.cpp b/Engine/Plugins/Developer/Concert/ConcertMain/Source/Concert/Private/ConcertMessageData.cpp
--- a/Engine/Plugins/Developer/Concert/ConcertMain/Source/Concert/Private/ConcertMessageData.cpp
+++ b/Engine/Plugins/Developer/Concert/ConcertMain/Source/Concert/Private/ConcertMessageData.cpp
@@ -335,12 +335,13 @@ bool FConcertSessionSerializedPayload::GetPayload(FStructOnScope& OutPayload) co
 bool FConcertSessionSerializedPayload::GetPayload(const UScriptStruct* InPayloadType, void* InOutPayloadData) const
 {
 	check(InPayloadType && InOutPayloadData);
+	return IsTypeChildOf(InPayloadType)
+		&& PayloadDetail::DeserializePreChecks((UScriptStruct*)InPayloadType, InOutPayloadData, *this)
+		&& PayloadDetail::DeserializeAndDecompress((UScriptStruct*)InPayloadType, InOutPayloadData, *this);
+}
+
+bool FConcertSessionSerializedPayload::IsTypeChildOf(const UScriptStruct* InPayloadType) const
+{
 	const UStruct* PayloadType = FindObject<UStruct>(nullptr, *PayloadTypeName.ToString());
-	if (PayloadType)
-	{
-		check(InPayloadType->IsChildOf(PayloadType));
-		return PayloadDetail::DeserializePreChecks((UScriptStruct*)InPayloadType, InOutPayloadData, *this)
-			&& PayloadDetail::DeserializeAndDecompress((UScriptStruct*)InPayloadType, InOutPayloadData, *this);
-	}
-	return false;
+	return PayloadType && InPayloadType->IsChildOf(PayloadType);
 }
diff --git a/Engine/Plugins/Developer/Concert/ConcertMain/Source/Concert/Private/ConcertServerSession.cpp b/Engine/Plugins/Developer/Concert/ConcertMain/Source/Concert/Private/ConcertServerSession.cpp
--- a/Engine/Plugins/Developer/Concert/ConcertMain/Source/Concert/Private/ConcertServerSession.cpp
+++ b/Engine/Plugins/Developer/Concert/ConcertMain/Source/Concert/Private/ConcertServerSession.cpp
@@ -94,6 +94,11 @@ FOnConcertServerSessionClientChanged& FConcertServerSession::OnSessionClientChan
 	return OnSessionClientChangedDelegate;
 }
 
+FOnConcertMessageAcknowledgementReceivedFromLocalEndpoint& FConcertServerSession::OnConcertMessageAcknowledgementReceived()
+{
+	return ServerSessionEndpoint->OnConcertMessageAcknowledgementReceived();
+}
+
 FString FConcertServerSession::GetSessionWorkingDirectory() const 
 {
 	return SessionDirectory;
diff --git a/Engine/Plugins/Developer/Concert/ConcertMain/Source/Concert/Private/ConcertServerSession.h b/Engine/Plugins/Developer/Concert/ConcertMain/Source/Concert/Private/ConcertServerSession.h
--- a/Engine/Plugins/Developer/Concert/ConcertMain/Source/Concert/Private/ConcertServerSession.h
+++ b/Engine/Plugins/Developer/Concert/ConcertMain/Source/Concert/Private/ConcertServerSession.h
@@ -81,6 +81,7 @@ public:
 
 	virtual FOnConcertServerSessionTick& OnTick() override;
 	virtual FOnConcertServerSessionClientChanged& OnSessionClientChanged() override;
+	virtual FOnConcertMessageAcknowledgementReceivedFromLocalEndpoint& OnConcertMessageAcknowledgementReceived() override;
 	virtual FString GetSessionWorkingDirectory() const override;
 
 protected:
diff --git a/Engine/Plugins/Developer/Concert/ConcertMain/Source/Concert/Private/ConcertSession.cpp b/Engine/Plugins/Developer/Concert/ConcertMain/Source/Concert/Private/ConcertSession.cpp
--- a/Engine/Plugins/Developer/Concert/ConcertMain/Source/Concert/Private/ConcertSession.cpp
+++ b/Engine/Plugins/Developer/Concert/ConcertMain/Source/Concert/Private/ConcertSession.cpp
@@ -118,7 +118,7 @@ void FConcertSessionCommonImpl::CommonHandleCustomEvent(const FConcertMessageCon
 		TArray<TSharedPtr<IConcertSessionCustomEventHandler>>* HandlerArrayPtr = CustomEventHandlers.Find(RawPayload.GetStruct()->GetFName());
 		if (HandlerArrayPtr)
 		{
-			FConcertSessionContext SessionContext{ Message->SourceEndpointId, Message->GetMessageFlags(), SenderScratchpad };
+			FConcertSessionContext SessionContext{ Message->SourceEndpointId, Message->MessageId, Message->GetMessageFlags(), SenderScratchpad };
 			for (const TSharedPtr<IConcertSessionCustomEventHandler>& Handler : *HandlerArrayPtr)
 			{
 				Handler->HandleEvent(SessionContext, RawPayload.GetStructMemory());
@@ -161,7 +161,7 @@ TFuture<FConcertSession_CustomResponse> FConcertSessionCommonImpl::CommonHandleC
 		if (Handler.IsValid())
 		{
 			FStructOnScope ResponsePayload(Handler->GetResponseType());
-			FConcertSessionContext SessionContext{ Message->SourceEndpointId, Message->GetMessageFlags(), SenderScratchpad };
+			FConcertSessionContext SessionContext{ Message->SourceEndpointId, Message->MessageId, Message->GetMessageFlags(), SenderScratchpad };
 			ResponseData.SetResponseCode(Handler->HandleRequest(SessionContext, RawPayload.GetStructMemory(), ResponsePayload.GetStructMemory()));
 			if (ResponseData.ResponseCode == EConcertResponseCode::Success || ResponseData.ResponseCode == EConcertResponseCode::Failed)
 			{
diff --git a/Engine/Plugins/Developer/Concert/ConcertMain/Source/Concert/Public/ConcertMessageData.h b/Engine/Plugins/Developer/Concert/ConcertMain/Source/Concert/Public/ConcertMessageData.h
--- a/Engine/Plugins/Developer/Concert/ConcertMain/Source/Concert/Public/ConcertMessageData.h
+++ b/Engine/Plugins/Developer/Concert/ConcertMain/Source/Concert/Public/ConcertMessageData.h
@@ -334,6 +334,14 @@ struct FConcertSessionSerializedPayload
 	CONCERT_API bool GetPayload(FStructOnScope& OutPayload) const;
 	CONCERT_API bool GetPayload(const UScriptStruct* InPayloadType, void* InOutPayloadData) const;
 
+	CONCERT_API bool IsTypeChildOf(const UScriptStruct* InPayloadType) const;
+	
+	template<typename T>
+	bool IsTypeChildOf() const
+	{
+		return IsTypeChildOf(TBaseStructure<T>::Get());
+	}
+
 	template <typename T>
 	bool GetTypedPayload(T& OutPayloadData) const
 	{
diff --git a/Engine/Plugins/Developer/Concert/ConcertMain/Source/Concert/Public/ConcertTransportEvents.h b/Engine/Plugins/Developer/Concert/ConcertMain/Source/Concert/Public/ConcertTransportEvents.h
--- a/Engine/Plugins/Developer/Concert/ConcertMain/Source/Concert/Public/ConcertTransportEvents.h
+++ b/Engine/Plugins/Developer/Concert/ConcertMain/Source/Concert/Public/ConcertTransportEvents.h
@@ -82,7 +82,9 @@ namespace ConcertTransportEvents
 	CONCERT_API void SetLoggingEnabled(bool bEnabled);
 
 	CONCERT_API FConcertTransportLoggingEnabledChanged& OnConcertTransportLoggingEnabledChangedEvent();
+	/** Use this when your local instance is a client. All client-generated logs are emitted by this event. */
 	CONCERT_API FConcertClientLogEvent& OnConcertClientLogEvent();
+	/** Use this when your local instance is a server. All server-generated logs are emitted by this event. */
 	CONCERT_API FConcertServerLogEvent& OnConcertServerLogEvent();
 }
 
diff --git a/Engine/Plugins/Developer/Concert/ConcertMain/Source/Concert/Public/IConcertSession.h b/Engine/Plugins/Developer/Concert/ConcertMain/Source/Concert/Public/IConcertSession.h
--- a/Engine/Plugins/Developer/Concert/ConcertMain/Source/Concert/Public/IConcertSession.h
+++ b/Engine/Plugins/Developer/Concert/ConcertMain/Source/Concert/Public/IConcertSession.h
@@ -98,6 +98,7 @@ public:
 	 * Get the scratchpad associated with the given client ID.
 	 */
 	virtual FConcertScratchpadPtr GetClientScratchpad(const FGuid& ClientEndpointId) const = 0;
+	
 
 	/**
 	 * Register a custom event handler for this session
@@ -255,6 +256,9 @@ public:
 
 	/** Callback when a session client state changes */
 	virtual FOnConcertServerSessionClientChanged& OnSessionClientChanged() = 0;
+	
+	/** Callback when a session message is acknowledged */
+	virtual FOnConcertMessageAcknowledgementReceivedFromLocalEndpoint& OnConcertMessageAcknowledgementReceived() = 0;
 };
 
 /** Interface for Concert client sessions */
diff --git a/Engine/Plugins/Developer/Concert/ConcertMain/Source/Concert/Public/IConcertSessionHandler.h b/Engine/Plugins/Developer/Concert/ConcertMain/Source/Concert/Public/IConcertSessionHandler.h
--- a/Engine/Plugins/Developer/Concert/ConcertMain/Source/Concert/Public/IConcertSessionHandler.h
+++ b/Engine/Plugins/Developer/Concert/ConcertMain/Source/Concert/Public/IConcertSessionHandler.h
@@ -13,6 +13,7 @@
 struct FConcertSessionContext
 {
 	FGuid SourceEndpointId;
+	FGuid MessageId;
 	EConcertMessageFlags MessageFlags;
 	FConcertScratchpadPtr SenderScratchpad;
 };
diff --git a/Engine/Plugins/Developer/Concert/ConcertSync/ConcertSyncClient/Source/ConcertSyncClient/Private/ConcertClientPackageManager.cpp b/Engine/Plugins/Developer/Concert/ConcertSync/ConcertSyncClient/Source/ConcertSyncClient/Private/ConcertClientPackageManager.cpp
--- a/Engine/Plugins/Developer/Concert/ConcertSync/ConcertSyncClient/Source/ConcertSyncClient/Private/ConcertClientPackageManager.cpp
+++ b/Engine/Plugins/Developer/Concert/ConcertSync/ConcertSyncClient/Source/ConcertSyncClient/Private/ConcertClientPackageManager.cpp
@@ -1,6 +1,7 @@
 // Copyright Epic Games, Inc. All Rights Reserved.
 
 #include "ConcertClientPackageManager.h"
+
 #include "IConcertSession.h"
 #include "IConcertFileSharingService.h"
 #include "ConcertSyncClientLiveSession.h"
@@ -506,8 +507,9 @@ void FConcertClientPackageManager::HandleLocalPackageEvent(const FConcertPackage
 	// if the package filter passes, send the event
 	if (ApplyPackageFilters(PackageInfo))
 	{
-		FConcertPackageUpdateEvent Event;
-		Event.Package.Info = PackageInfo;
+		FConcertPackageUpdateEvent PackageEvent;
+		PackageEvent.TransmissionId = FGuid::NewGuid();
+		PackageEvent.Package.Info = PackageInfo;
 
 		// Copy or link the package data to the Concert event.
 		int64 PackageFileSize = PackagePathname.IsEmpty() ? -1 : IFileManager::Get().FileSize(*PackagePathname); // EConcertPackageUpdateType::Delete is emitted with an empty pathname.
@@ -516,7 +518,7 @@ void FConcertClientPackageManager::HandleLocalPackageEvent(const FConcertPackage
 			if (CanExchangePackageDataAsByteArray(static_cast<uint64>(PackageFileSize)))
 			{
 				// Embed the package data directly in the event.
-				if (!FFileHelper::LoadFileToArray(Event.Package.PackageData.Bytes, *PackagePathname))
+				if (!FFileHelper::LoadFileToArray(PackageEvent.Package.PackageData.Bytes, *PackagePathname))
 				{
 					UE_LOG(LogConcert, Error, TEXT("Failed to load file data '%s' in memory"), *PackagePathname);
 					return;
@@ -525,7 +527,7 @@ void FConcertClientPackageManager::HandleLocalPackageEvent(const FConcertPackage
 			else if (FileSharingService && EnumHasAnyFlags(LiveSession->GetSessionFlags(), EConcertSyncSessionFlags::EnableFileSharing))
 			{
 				// Publish a copy of the package data in the sharing service and set the corresponding file ID in the event.
-				if (!FileSharingService->Publish(PackagePathname, Event.Package.FileId))
+				if (!FileSharingService->Publish(PackagePathname, PackageEvent.Package.FileId))
 				{
 					UE_LOG(LogConcert, Error, TEXT("Failed to share a copy of package file '%s'"), *PackagePathname);
 					return;
@@ -535,12 +537,17 @@ void FConcertClientPackageManager::HandleLocalPackageEvent(const FConcertPackage
 			{
 				// Notify the client about the file being too large to be emitted.
 				UE_LOG(LogConcert, Error, TEXT("Failed to handle local package file '%s'. The file is too big to be sent over the network."), *PackagePathname);
-				OnConcertClientPackageTooLargeError().Broadcast(Event.Package.Info, PackageFileSize, FConcertPackage::GetMaxPackageDataSizeEmbeddableAsByteArray());
+				OnConcertClientPackageTooLargeError().Broadcast(PackageEvent.Package.Info, PackageFileSize, FConcertPackage::GetMaxPackageDataSizeEmbeddableAsByteArray());
 			}
 		}
 
-		LiveSession->GetSessionDatabase().GetTransactionMaxEventId(Event.Package.Info.TransactionEventIdAtSave);
-		LiveSession->GetSession().SendCustomEvent(Event, LiveSession->GetSession().GetSessionServerEndpointId(), EConcertMessageFlags::ReliableOrdered);
+		LiveSession->GetSessionDatabase().GetTransactionMaxEventId(PackageEvent.Package.Info.TransactionEventIdAtSave);
+		if (PackageEvent.Package.HasPackageData())
+		{
+			const FConcertPackageTransmissionStartEvent AnnouncementEvent{ PackageEvent.TransmissionId, PackageEvent.Package.Info, static_cast<uint64>(PackageFileSize) };
+			LiveSession->GetSession().SendCustomEvent(AnnouncementEvent, LiveSession->GetSession().GetSessionServerEndpointId(), EConcertMessageFlags::ReliableOrdered);
+		}
+		LiveSession->GetSession().SendCustomEvent(PackageEvent, LiveSession->GetSession().GetSessionServerEndpointId(), EConcertMessageFlags::ReliableOrdered);
 	}
 	// if the package data has been filtered out of the session persist it immediately from the sandbox
 	else
diff --git a/Engine/Plugins/Developer/Concert/ConcertSync/ConcertSyncCore/Source/ConcertSyncCore/Private/ConcertPackageEvents.cpp b/Engine/Plugins/Developer/Concert/ConcertSync/ConcertSyncCore/Source/ConcertSyncCore/Private/ConcertPackageEvents.cpp
new file mode 100644
--- /dev/null
+++ b/Engine/Plugins/Developer/Concert/ConcertSync/ConcertSyncCore/Source/ConcertSyncCore/Private/ConcertPackageEvents.cpp
@@ -0,0 +1,36 @@
+// Copyright Epic Games, Inc. All Rights Reserved.
+
+#include "ConcertPackageEvents.h"
+
+namespace UE::ConcertSyncCore::ConcertPackageEvents
+{
+	FConcertBeginSendPackageDelegate& OnLocalBeginSendPackage()
+	{
+		static FConcertBeginSendPackageDelegate Instance;
+		return Instance;
+	}
+	
+	FConcertFinishSendPackageDelegate& OnLocalFinishSendPackage()
+	{
+		static FConcertFinishSendPackageDelegate Instance;
+		return Instance;
+	}
+
+	FConcertBeginSendPackageDelegate& OnRemoteBeginSendPackage()
+	{
+		static FConcertBeginSendPackageDelegate Instance;
+		return Instance;
+	}
+
+	FConcertFinishSendPackageDelegate& OnRemoteFinishSendPackage()
+	{
+		static FConcertFinishSendPackageDelegate Instance;
+		return Instance;
+	}
+
+	FConcertRejectSendPackageDelegate& OnRejectRemoteSendPackage()
+	{
+		static FConcertRejectSendPackageDelegate Instance;
+		return Instance;
+	}
+}
\ No newline at end of file
diff --git a/Engine/Plugins/Developer/Concert/ConcertSync/ConcertSyncCore/Source/ConcertSyncCore/Public/ConcertPackageEvents.h b/Engine/Plugins/Developer/Concert/ConcertSync/ConcertSyncCore/Source/ConcertSyncCore/Public/ConcertPackageEvents.h
new file mode 100644
--- /dev/null
+++ b/Engine/Plugins/Developer/Concert/ConcertSync/ConcertSyncCore/Source/ConcertSyncCore/Public/ConcertPackageEvents.h
@@ -0,0 +1,67 @@
+// Copyright Epic Games, Inc. All Rights Reserved.
+
+#pragma once
+
+#include "CoreMinimal.h"
+
+struct FConcertPackageInfo;
+
+namespace UE::ConcertSyncCore::ConcertPackageEvents
+{
+	struct FConcertBaseSendPackageParams
+	{
+		/** Identifies begin and finish send events so you can map them to each other. */
+		FGuid TransmissionId;
+		
+		/** General information about the package */
+		const FConcertPackageInfo& PackageInfo;
+		
+		/**
+		 * The remote endpoint that is interacting with this package.
+		 *
+		 * If we're the server, then there are 2 options:
+		 * 1. This is a client that's sending us data
+		 * 2. This is a client that we're sending the data to
+		 */
+		FGuid RemoteEndpointId;
+	};
+	
+	struct FConcertBeginSendPackageParams : FConcertBaseSendPackageParams
+	{
+		/** The size of the package data to be received */
+		uint64 PackageNumBytes;
+	};
+	
+	struct FConcertFinishSendPackageParams : FConcertBaseSendPackageParams
+	{
+		/**
+		 * The ID of the message that causes the package to be acknowledged.
+		 *
+		 * In the case of sending to a client, it is the acknowledgement message.
+		 * In the case of receiving from client, it is the message transmitting the sync activity.
+		 */
+		FGuid MessageId;
+	};
+
+	struct FConcertRejectSendPackageParams : FConcertBaseSendPackageParams
+	{};
+	
+	DECLARE_MULTICAST_DELEGATE_OneParam(FConcertBeginSendPackageDelegate, const FConcertBeginSendPackageParams&);
+	DECLARE_MULTICAST_DELEGATE_OneParam(FConcertFinishSendPackageDelegate, const FConcertFinishSendPackageParams&);
+	DECLARE_MULTICAST_DELEGATE_OneParam(FConcertRejectSendPackageDelegate, const FConcertRejectSendPackageParams&);
+	
+	// Note: These events are only called when package contents is sent. Package activities do not always contain package data in which case these events are not called.
+	// Caution: For now these events are only called on the server application; they're not called in client applications.
+	
+	/** Called when we begin sending a package to a remote endpoint. RemoteEndpointId is to be interpreted as the receiver. */
+	CONCERTSYNCCORE_API FConcertBeginSendPackageDelegate& OnLocalBeginSendPackage();
+	/** Called when we finish sending a package to a remote endpoint. They've sent us an ACK. RemoteEndpointId is to be interpreted as the receiver. */
+	CONCERTSYNCCORE_API FConcertFinishSendPackageDelegate& OnLocalFinishSendPackage();
+
+	/** Called when we begin receiving a package from a remote endpoint. They've let us know in advance via an event. RemoteEndpointId is to be interpreted as the sender. */
+	CONCERTSYNCCORE_API FConcertBeginSendPackageDelegate& OnRemoteBeginSendPackage();
+	/** Called when we finished receiving a package from a remote endpoint. We've processed the data locally. RemoteEndpointId is to be interpreted as the sender. */
+	CONCERTSYNCCORE_API FConcertFinishSendPackageDelegate& OnRemoteFinishSendPackage();
+	/** Called when we finished receiving a package from a remote endpoint but decided to reject the data. Called only on the server. RemoteEndpointId is to be interpreted as the sender. */
+	CONCERTSYNCCORE_API FConcertRejectSendPackageDelegate& OnRejectRemoteSendPackage();
+}
diff --git a/Engine/Plugins/Developer/Concert/ConcertSync/ConcertSyncCore/Source/ConcertSyncCore/Public/ConcertWorkspaceMessages.h b/Engine/Plugins/Developer/Concert/ConcertSync/ConcertSyncCore/Source/ConcertSyncCore/Public/ConcertWorkspaceMessages.h
--- a/Engine/Plugins/Developer/Concert/ConcertSync/ConcertSyncCore/Source/ConcertSyncCore/Public/ConcertWorkspaceMessages.h
+++ b/Engine/Plugins/Developer/Concert/ConcertSync/ConcertSyncCore/Source/ConcertSyncCore/Public/ConcertWorkspaceMessages.h
@@ -108,11 +108,31 @@ struct FConcertServerLogging
 	bool bLoggingEnabled = false;
 };
 
+/** Sent to let the receiver know something is coming. For now only sent from client to server. */
+USTRUCT()
+struct FConcertPackageTransmissionStartEvent
+{
+	GENERATED_BODY()
+
+	UPROPERTY()
+	FGuid TransmissionId;
+	
+	UPROPERTY()
+	FConcertPackageInfo PackageInfo;
+
+	UPROPERTY()
+	uint64 PackageNumBytes;
+};
+
+
 USTRUCT()
 struct FConcertPackageUpdateEvent
 {
 	GENERATED_BODY()
 
+	UPROPERTY()
+	FGuid TransmissionId;
+	
 	UPROPERTY()
 	FConcertPackage Package;
 };
diff --git a/Engine/Plugins/Developer/Concert/ConcertSync/ConcertSyncServer/Source/ConcertSyncServer/Private/ConcertServerWorkspace.cpp b/Engine/Plugins/Developer/Concert/ConcertSync/ConcertSyncServer/Source/ConcertSyncServer/Private/ConcertServerWorkspace.cpp
--- a/Engine/Plugins/Developer/Concert/ConcertSync/ConcertSyncServer/Source/ConcertSyncServer/Private/ConcertServerWorkspace.cpp
+++ b/Engine/Plugins/Developer/Concert/ConcertSync/ConcertSyncServer/Source/ConcertSyncServer/Private/ConcertServerWorkspace.cpp
@@ -11,9 +11,11 @@
 #include "ConcertTransactionEvents.h"
 #include "ConcertServerDataStore.h"
 #include "ConcertLogGlobal.h"
+#include "ConcertPackageEvents.h"
 #include "ConcertUtil.h"
 #include "Serialization/MemoryReader.h"
 #include "Algo/Transform.h"
+#include "HistoryEdition/ActivityDependencyGraph.h"
 
 FConcertServerWorkspace::FConcertServerWorkspace(const TSharedRef<FConcertSyncServerLiveSession>& InLiveSession, TSharedPtr<IConcertFileSharingService> InFileSharingService)
 	: FileSharingService(MoveTemp(InFileSharingService))
@@ -50,9 +52,12 @@ void FConcertServerWorkspace::BindSession(const TSharedRef<FConcertSyncServerLiv
 
 	// Register Client Change events
 	LiveSession->GetSession().OnSessionClientChanged().AddRaw(this, &FConcertServerWorkspace::HandleSessionClientChanged);
+	
+	LiveSession->GetSession().OnConcertMessageAcknowledgementReceived().AddRaw(this, &FConcertServerWorkspace::HandleSessionAcknowledgementReceived);
 
 	LiveSession->GetSession().RegisterCustomEventHandler<FConcertWorkspaceSyncRequestedEvent>(this, &FConcertServerWorkspace::HandleSyncRequestedEvent);
 
+	LiveSession->GetSession().RegisterCustomEventHandler<FConcertPackageTransmissionStartEvent>(this, &FConcertServerWorkspace::HandlePackageTransmissionStartEvent);
 	LiveSession->GetSession().RegisterCustomEventHandler<FConcertPackageUpdateEvent>(this, &FConcertServerWorkspace::HandlePackageUpdateEvent);
 	LiveSession->GetSession().RegisterCustomEventHandler<FConcertPlaySessionEvent>(this, &FConcertServerWorkspace::HandlePlaySessionEvent);
 
@@ -63,6 +68,7 @@ void FConcertServerWorkspace::BindSession(const TSharedRef<FConcertSyncServerLiv
 	LiveSession->GetSession().RegisterCustomRequestHandler<FConcertSyncEventRequest, FConcertSyncEventResponse>(this, &FConcertServerWorkspace::HandleSyncEventRequest);
 	LiveSession->GetSession().RegisterCustomEventHandler<FConcertIgnoreActivityStateChangedEvent>(this, &FConcertServerWorkspace::HandleIgnoredActivityStateChanged);
 	LiveSession->GetSession().RegisterCustomEventHandler<FConcertServerLogging>(this, &FConcertServerWorkspace::HandleServerLoggingEvent);
+
 }
 
 void FConcertServerWorkspace::UnbindSession()
@@ -84,8 +90,11 @@ void FConcertServerWorkspace::UnbindSession()
 		// Unregister Client Change events
 		LiveSession->GetSession().OnSessionClientChanged().RemoveAll(this);
 
+		LiveSession->GetSession().OnConcertMessageAcknowledgementReceived().RemoveAll(this);
+
 		LiveSession->GetSession().UnregisterCustomEventHandler<FConcertWorkspaceSyncRequestedEvent>(this);
 
+		LiveSession->GetSession().UnregisterCustomEventHandler<FConcertPackageTransmissionStartEvent>(this);
 		LiveSession->GetSession().UnregisterCustomEventHandler<FConcertPackageUpdateEvent>(this);
 		LiveSession->GetSession().UnregisterCustomEventHandler<FConcertPlaySessionEvent>(this);
 
@@ -174,6 +183,41 @@ void FConcertServerWorkspace::HandleSessionClientChanged(IConcertServerSession&
 	}
 }
 
+namespace UE::ConcertSyncServer::Private
+{
+	template<typename TSyncActivityType>
+	static void ExtractSyncActivity(const TSharedRef<IConcertMessage>& AckedMessage, TFunctionRef<void(const TSyncActivityType& Activity)> Callback)
+	{
+		if (AckedMessage->GetMessageType() == FConcertSession_CustomEvent::StaticStruct())
+		{
+			const FConcertSession_CustomEvent* CustomEvent = reinterpret_cast<const FConcertSession_CustomEvent*>(AckedMessage->ConstructMessage());
+			if (CustomEvent->SerializedPayload.IsTypeChildOf<FConcertWorkspaceSyncActivityEvent>())
+			{
+				FConcertWorkspaceSyncActivityEvent SyncActivity;
+				CustomEvent->SerializedPayload.GetTypedPayload(SyncActivity);
+				if (SyncActivity.Activity.IsTypeChildOf(TBaseStructure<TSyncActivityType>::Get()))
+				{
+					TSyncActivityType Result;
+					SyncActivity.Activity.GetTypedPayload(Result);
+					Callback(Result);
+				}
+			}
+		}
+	}
+}
+
+void FConcertServerWorkspace::HandleSessionAcknowledgementReceived(const FConcertEndpointContext& LocalEndpoint, const FConcertEndpointContext& RemoteEndpoint, const TSharedRef<IConcertMessage>& AckedMessage, const FConcertMessageContext& MessageContext) const
+{
+	UE::ConcertSyncServer::Private::ExtractSyncActivity<FConcertSyncPackageActivity>(AckedMessage, [&RemoteEndpoint, &AckedMessage](const FConcertSyncPackageActivity& PackageActivity)
+	{
+		if (PackageActivity.EventData.Package.HasPackageData())
+		{
+			const FGuid TransmissionId = HashEndpointIdAndActivityId(RemoteEndpoint.EndpointId, PackageActivity.ActivityId);
+			UE::ConcertSyncCore::ConcertPackageEvents::OnLocalFinishSendPackage().Broadcast({ { TransmissionId, PackageActivity.EventData.Package.Info, RemoteEndpoint.EndpointId }, AckedMessage->GetMessageId() });
+		}
+	});
+}
+
 void FConcertServerWorkspace::HandleIgnoredActivityStateChanged(const FConcertSessionContext& Context, const FConcertIgnoreActivityStateChangedEvent& Event)
 {
 	if (Event.bIgnore)
@@ -297,6 +341,16 @@ void FConcertServerWorkspace::HandleSyncRequestedEvent(const FConcertSessionCont
 	});
 }
 
+void FConcertServerWorkspace::HandlePackageTransmissionStartEvent(const FConcertSessionContext& Context, const FConcertPackageTransmissionStartEvent& Event)
+{
+	if (EnumHasAnyFlags(LiveSession->GetSessionFlags(), EConcertSyncSessionFlags::EnablePackages)
+		&& ensureMsgf(Event.PackageNumBytes > 0, TEXT("Client sent invalid message")))
+	{
+		// Clients send this before the FConcertPackageUpdateEvent in reliable order so we can be sure we'll trigger OnRemoteBeginSendPackage before OnRemoteFinishSendPackage
+		UE::ConcertSyncCore::ConcertPackageEvents::OnRemoteBeginSendPackage().Broadcast({ { Event.TransmissionId, Event.PackageInfo, Context.SourceEndpointId }, Event.PackageNumBytes });
+	}
+}
+
 void FConcertServerWorkspace::HandlePackageUpdateEvent(const FConcertSessionContext& Context, const FConcertPackageUpdateEvent& Event)
 {
 	if (!EnumHasAnyFlags(LiveSession->GetSessionFlags(), EConcertSyncSessionFlags::EnablePackages))
@@ -351,16 +405,26 @@ void FConcertServerWorkspace::HandlePackageUpdateEvent(const FConcertSessionCont
 					PackageDataAr ? PackageDataAr->TotalSize() : 0,
 					Event.Package.PackageData.Bytes.Num() ? &Event.Package.PackageData.Bytes : nullptr
 				};
-
+ 
 				AddPackageActivity(PackageActivityBasePart, Package.Info, PackageDataStream);
 			}
 		}
 
+		if (Event.Package.HasPackageData())
+		{
+			UE::ConcertSyncCore::ConcertPackageEvents::OnRemoteFinishSendPackage().Broadcast({ { Event.TransmissionId, Event.Package.Info, Context.SourceEndpointId }, Context.MessageId });
+		}
+		
 		// Explicitly unlock the resource after saving it
 		UnlockWorkspaceResource(Event.Package.Info.PackageName, Context.SourceEndpointId, EConcertLockFlags::Explicit);
 	}
 	else
 	{
+		if (Event.Package.HasPackageData())
+		{
+			UE::ConcertSyncCore::ConcertPackageEvents::OnRejectRemoteSendPackage().Broadcast({ Event.TransmissionId, Event.Package.Info, Context.SourceEndpointId });
+		}
+		
 		// If the client didn't have the lock, then queue a rejection event so that the client will re-load the head-revision of the package
 		SyncCommandQueue->QueueCommand(Context.SourceEndpointId, [this, PackageName = Event.Package.Info.PackageName](const FConcertServerSyncCommandQueue::FSyncCommandContext& InSyncCommandContext, const FGuid& InEndpointId)
 		{
@@ -932,6 +996,13 @@ void FConcertServerWorkspace::AddTransactionActivity(const FConcertSyncTransacti
 	}
 }
 
+FGuid FConcertServerWorkspace::HashEndpointIdAndActivityId(const FGuid& TargetEndpointId, FActivityID ActivityID)
+{
+	const int32 LowBits = static_cast<int32>(ActivityID);
+	const int32 HighBits = (ActivityID >> 32);
+	return FGuid(TargetEndpointId.A + LowBits, TargetEndpointId.B + HighBits, TargetEndpointId.C, TargetEndpointId.D);
+}
+
 TOptional<FConcertWorkspaceSyncActivityEvent> FConcertServerWorkspace::SyncTransactionActivityEvent(const int64 InSyncActivityId, const int32 InNumRemainingSyncEvents, const bool InLiveOnly) const
 {
 	FConcertSyncTransactionActivity SyncActivity;
@@ -991,7 +1062,6 @@ void FConcertServerWorkspace::AddPackageActivity(const FConcertSyncActivity& InP
 				SendSyncPackageActivityEvent(Event.GetValue(), InEndpointId);
 			});
 		}
-
 	}
 	else
 	{
@@ -1057,6 +1127,25 @@ TOptional<FConcertWorkspaceSyncActivityEvent> FConcertServerWorkspace::MakeSyncA
 void FConcertServerWorkspace::SendSyncPackageActivityEvent(const FConcertWorkspaceSyncActivityEvent& SyncEvent, const FGuid& InTargetEndpointId) const
 {
 	SCOPED_CONCERT_TRACE(FConcertServerWorkspace_SendSyncPackageActivityEvent);
+
+	if (UE::ConcertSyncCore::ConcertPackageEvents::OnLocalBeginSendPackage().IsBound())
+	{
+		FConcertSyncPackageActivity PackageActivity;
+		SyncEvent.Activity.GetTypedPayload(PackageActivity);
+		const FGuid TransmissionId = HashEndpointIdAndActivityId(InTargetEndpointId, PackageActivity.ActivityId);
+		UE::ConcertSyncCore::ConcertPackageEvents::FConcertBeginSendPackageParams PreTransmitEvent{ { TransmissionId, PackageActivity.EventData.Package.Info, InTargetEndpointId } };
+		// This loads the file even though we only want the file size... could be improved
+		LiveSession->GetSessionDatabase().GetPackageEvent(PackageActivity.EventId, [&PreTransmitEvent](const FConcertSyncPackageEventData& Data)
+		{
+			PreTransmitEvent.PackageNumBytes = static_cast<uint64>(Data.PackageDataStream.DataSize);
+		});
+		
+		if (PackageActivity.EventData.Package.HasPackageData())
+		{
+			UE::ConcertSyncCore::ConcertPackageEvents::OnLocalBeginSendPackage().Broadcast(PreTransmitEvent);
+		}
+	}
+	
 	LiveSession->GetSession().SendCustomEvent(SyncEvent, InTargetEndpointId, EConcertMessageFlags::ReliableOrdered);
 }
 
diff --git a/Engine/Plugins/Developer/Concert/ConcertSync/ConcertSyncServer/Source/ConcertSyncServer/Private/ConcertServerWorkspace.h b/Engine/Plugins/Developer/Concert/ConcertSync/ConcertSyncServer/Source/ConcertSyncServer/Private/ConcertServerWorkspace.h
--- a/Engine/Plugins/Developer/Concert/ConcertSync/ConcertSyncServer/Source/ConcertSyncServer/Private/ConcertServerWorkspace.h
+++ b/Engine/Plugins/Developer/Concert/ConcertSync/ConcertSyncServer/Source/ConcertSyncServer/Private/ConcertServerWorkspace.h
@@ -6,11 +6,16 @@
 #include "IConcertSessionHandler.h"
 #include "ConcertWorkspaceMessages.h"
 #include "Misc/Optional.h"
-class IConcertServerSession;
+
 class IConcertFileSharingService;
+class IConcertMessage;
+class IConcertServerSession;
 class FConcertSyncServerLiveSession;
 class FConcertServerSyncCommandQueue;
 class FConcertServerDataStore;
+
+struct FConcertMessageContext;
+struct FConcertEndpointContext;
 struct FConcertTransactionSnapshotEvent;
 struct FConcertTransactionFinalizedEvent;
 
@@ -30,6 +35,10 @@ public:
 	~FConcertServerWorkspace();
 
 private:
+	
+	/** Creates an ID from the given data through an simple hash function. The ID only depends on the input, allowing us to reconstruct IDs later in time. */
+	static FGuid HashEndpointIdAndActivityId(const FGuid& TargetEndpointId, FActivityID ActivityID);
+	
 	/**
 	 * Get the sync event for a transaction activity in the session database.
 	 *
@@ -61,9 +70,15 @@ private:
 	/** */
 	void HandleSessionClientChanged(IConcertServerSession& InSession, EConcertClientStatus InClientStatus, const FConcertSessionClientInfo& InClientInfo);
 
+	/** */
+	void HandleSessionAcknowledgementReceived(const FConcertEndpointContext& LocalEndpoint, const FConcertEndpointContext& RemoteEndpoint, const TSharedRef<IConcertMessage>& AckedMessage, const FConcertMessageContext& MessageContext) const;
+
 	/** */
 	void HandleSyncRequestedEvent(const FConcertSessionContext& Context, const FConcertWorkspaceSyncRequestedEvent& Event);
 
+	/** */
+	void HandlePackageTransmissionStartEvent(const FConcertSessionContext& Context, const FConcertPackageTransmissionStartEvent& Event);
+	
 	/** */
 	void HandlePackageUpdateEvent(const FConcertSessionContext& Context, const FConcertPackageUpdateEvent& Event);
 
diff --git a/Engine/Plugins/Developer/Concert/ConcertSync/ConcertSyncTest/Source/ConcertSyncTest/Private/ConcertDataStoreTests.cpp b/Engine/Plugins/Developer/Concert/ConcertSync/ConcertSyncTest/Source/ConcertSyncTest/Private/ConcertDataStoreTests.cpp
--- a/Engine/Plugins/Developer/Concert/ConcertSync/ConcertSyncTest/Source/ConcertSyncTest/Private/ConcertDataStoreTests.cpp
+++ b/Engine/Plugins/Developer/Concert/ConcertSync/ConcertSyncTest/Source/ConcertSyncTest/Private/ConcertDataStoreTests.cpp
@@ -103,6 +103,7 @@ public:
 	// IConcertServerSession Begin
 	virtual FOnConcertServerSessionTick& OnTick() override                          { return NotMocked<FOnConcertServerSessionTick&>(Tick); }
 	virtual FOnConcertServerSessionClientChanged& OnSessionClientChanged() override { return NotMocked<FOnConcertServerSessionClientChanged&>(ConnectionChanged); }
+	virtual FOnConcertMessageAcknowledgementReceivedFromLocalEndpoint& OnConcertMessageAcknowledgementReceived() override { return NotMocked<FOnConcertMessageAcknowledgementReceivedFromLocalEndpoint&>(AckReceived); }
 	// IConcertServerSession End
 
 protected:
@@ -113,6 +114,7 @@ protected:
 	FConcertSessionInfo SessionInfo;
 	FOnConcertServerSessionTick Tick;
 	FOnConcertServerSessionClientChanged ConnectionChanged;
+	FOnConcertMessageAcknowledgementReceivedFromLocalEndpoint AckReceived;
 };
 
 /** Implements a not-working IConcertClientSession. It must be further overridden to implement just what is required by the tests */
