Fix CIS issue

diff --git a/Engine/Plugins/Developer/Concert/ConcertSync/ConcertSyncServer/Source/ConcertSyncServer/Private/ConcertServerWorkspace.cpp b/Engine/Plugins/Developer/Concert/ConcertSync/ConcertSyncServer/Source/ConcertSyncServer/Private/ConcertServerWorkspace.cpp
--- a/Engine/Plugins/Developer/Concert/ConcertSync/ConcertSyncServer/Source/ConcertSyncServer/Private/ConcertServerWorkspace.cpp
+++ b/Engine/Plugins/Developer/Concert/ConcertSync/ConcertSyncServer/Source/ConcertSyncServer/Private/ConcertServerWorkspace.cpp
@@ -422,7 +422,7 @@ void FConcertServerWorkspace::HandlePackageUpdateEvent(const FConcertSessionCont
 	{
 		if (Event.Package.HasPackageData())
 		{
-			UE::ConcertSyncCore::ConcertPackageEvents::OnRejectRemoteSendPackage().Broadcast({ Event.TransmissionId, Event.Package.Info, Context.SourceEndpointId });
+			UE::ConcertSyncCore::ConcertPackageEvents::OnRejectRemoteSendPackage().Broadcast({{ Event.TransmissionId, Event.Package.Info, Context.SourceEndpointId }});
 		}
 		
 		// If the client didn't have the lock, then queue a rejection event so that the client will re-load the head-revision of the package
