Backing out addition of Launch dependency until we figure out the dependency loop.

diff --git a/Engine/Plugins/Developer/Concert/ConcertSync/ConcertSyncServer/Source/ConcertSyncServer/ConcertSyncServer.Build.cs b/Engine/Plugins/Developer/Concert/ConcertSync/ConcertSyncServer/Source/ConcertSyncServer/ConcertSyncServer.Build.cs
--- a/Engine/Plugins/Developer/Concert/ConcertSync/ConcertSyncServer/Source/ConcertSyncServer/ConcertSyncServer.Build.cs
+++ b/Engine/Plugins/Developer/Concert/ConcertSync/ConcertSyncServer/Source/ConcertSyncServer/ConcertSyncServer.Build.cs
@@ -16,7 +16,6 @@ namespace UnrealBuildTool.Rules
 					"CoreUObject",
 					"Concert",
 					"ConcertSyncCore",
-					"Launch",
 				}
 			);
 
