Fix UnrealMultiUserSlateServer crashing when closed because FEnginLoop::AppExit was called twice

diff --git a/Engine/Plugins/Developer/Concert/ConcertSync/ConcertSyncServer/Source/ConcertSyncServer/Public/ConcertSyncServerLoop.inl b/Engine/Plugins/Developer/Concert/ConcertSync/ConcertSyncServer/Source/ConcertSyncServer/Public/ConcertSyncServerLoop.inl
--- a/Engine/Plugins/Developer/Concert/ConcertSync/ConcertSyncServer/Source/ConcertSyncServer/Public/ConcertSyncServerLoop.inl
+++ b/Engine/Plugins/Developer/Concert/ConcertSync/ConcertSyncServer/Source/ConcertSyncServer/Public/ConcertSyncServerLoop.inl
@@ -191,7 +191,10 @@ int32 ConcertSyncServerLoop(const TCHAR* CommandLine, const FConcertSyncServerLo
 	// Unloading Modules isn't handled by AppExit
 	FModuleManager::Get().UnloadModulesAtShutdown();
 
+	// On Mac AppExit is passed as delegate to the OS, which will call it implicitly
+#if !PLATFORM_MAC
 	FEngineLoop::AppExit();
+#endif
 
 	return Result;
 }
