ConcertSync + Sequencer Playlists: Make sequence preloading use asynchronous package loading. Indicate preloading result of other multi-user clients in Playlists UI.

diff --git a/Engine/Plugins/Developer/Concert/ConcertSync/ConcertSyncClient/Source/ConcertSyncClient/Private/ConcertClientSequencerManager.cpp b/Engine/Plugins/Developer/Concert/ConcertSync/ConcertSyncClient/Source/ConcertSyncClient/Private/ConcertClientSequencerManager.cpp
--- a/Engine/Plugins/Developer/Concert/ConcertSync/ConcertSyncClient/Source/ConcertSyncClient/Private/ConcertClientSequencerManager.cpp
+++ b/Engine/Plugins/Developer/Concert/ConcertSync/ConcertSyncClient/Source/ConcertSyncClient/Private/ConcertClientSequencerManager.cpp
@@ -8,6 +8,7 @@
 #include "Misc/QualifiedFrameTime.h"
 #include "Modules/ModuleManager.h"
 #include "Logging/LogMacros.h"
+#include "Logging/StructuredLog.h"
 
 #include "IConcertClient.h"
 #include "IConcertSyncClient.h"
@@ -60,8 +61,53 @@ static TAutoConsoleVariable<int32> CVarEnableUnrelatedTimelineSync(TEXT("Concert
 static TAutoConsoleVariable<int32> CVarAlwaysCloseGamePlayerOnCloseEvent(
 	TEXT("Concert.AlwaysCloseGamePlayerOnCloseEvent"), 1, TEXT("Force this player to close even if other editors have it open. This CVar only works on `-game` instances."));
 
+
+class FConcertClientSequencePreloader : public TSharedFromThis<FConcertClientSequencePreloader>
+{
+public:
+	void OnPreloadEvent(const FConcertSessionContext& InEventContext, const FConcertSequencerPreloadRequest& InEvent);
+
+	void OnRegister(TSharedRef<IConcertClientSession> InSession)
+	{
+		WeakSession = InSession;
+
+		InSession->RegisterCustomEventHandler<FConcertSequencerPreloadRequest>(this, &FConcertClientSequencePreloader::OnPreloadEvent);
+	}
+
+	void OnUnregister(TSharedRef<IConcertClientSession> InSession)
+	{
+		PreloadPendingSequences.Empty();
+		PreloadedSequences.Empty();
+
+		// Unregister our events and explicitly reset the session ptr
+		if (TSharedPtr<IConcertClientSession> Session = WeakSession.Pin())
+		{
+			check(Session == InSession);
+			Session->UnregisterCustomEventHandler<FConcertSequencerPreloadRequest>(this);
+		}
+
+		WeakSession.Reset();
+	}
+
+	void AddReferencedObjects(FReferenceCollector& Collector)
+	{
+		Collector.AddReferencedObjects(PreloadedSequences);
+	}
+
+protected:
+	TWeakPtr<IConcertClientSession> WeakSession;
+
+	/** Map of sequence assets for which preload has been requested but not completed. */
+	TMap<FTopLevelAssetPath, TSoftObjectPtr<ULevelSequence>> PreloadPendingSequences;
+
+	/** Map of preloaded sequence assets. */
+	TMap<FTopLevelAssetPath, TObjectPtr<ULevelSequence>> PreloadedSequences;
+};
+
+
 FConcertClientSequencerManager::FConcertClientSequencerManager(IConcertSyncClient* InOwnerSyncClient)
 	: OwnerSyncClient(InOwnerSyncClient)
+	, Preloader(MakeShared<FConcertClientSequencePreloader>())
 {
 	check(OwnerSyncClient);
 
@@ -199,13 +245,12 @@ void FConcertClientSequencerManager::Register(TSharedRef<IConcertClientSession>
 	InSession->RegisterCustomEventHandler<FConcertSequencerOpenEvent>(this, &FConcertClientSequencerManager::OnOpenEvent);
 	InSession->RegisterCustomEventHandler<FConcertSequencerStateSyncEvent>(this, &FConcertClientSequencerManager::OnStateSyncEvent);
 	InSession->RegisterCustomEventHandler<FConcertSequencerTimeAdjustmentEvent>(this, &FConcertClientSequencerManager::OnTimeAdjustmentEvent);
-	InSession->RegisterCustomEventHandler<FConcertSequencerPrecacheEvent>(this, &FConcertClientSequencerManager::OnPrecacheEvent);
+
+	Preloader->OnRegister(InSession);
 }
 
 void FConcertClientSequencerManager::Unregister(TSharedRef<IConcertClientSession> InSession)
 {
-	PrecachedSequences.Empty();
-
 	// Unregister our events and explicitly reset the session ptr
 	if (TSharedPtr<IConcertClientSession> Session = WeakSession.Pin())
 	{
@@ -215,9 +260,12 @@ void FConcertClientSequencerManager::Unregister(TSharedRef<IConcertClientSession
 		Session->UnregisterCustomEventHandler<FConcertSequencerOpenEvent>(this);
 		Session->UnregisterCustomEventHandler<FConcertSequencerStateSyncEvent>(this);
 		Session->UnregisterCustomEventHandler<FConcertSequencerTimeAdjustmentEvent>(this);
-		Session->UnregisterCustomEventHandler<FConcertSequencerPrecacheEvent>(this);
+		Session->UnregisterCustomEventHandler<FConcertSequencerPreloadRequest>(this);
 	}
+
 	WeakSession.Reset();
+
+	Preloader->OnUnregister(InSession);
 }
 
 void SetConsoleVariableRespectingPriority(IConsoleVariable* AsVariable, bool bValue)
@@ -814,37 +862,111 @@ void FConcertClientSequencerManager::ApplyTimeAdjustmentToPlayers(const FConcert
 	}
 }
 
-void FConcertClientSequencerManager::OnPrecacheEvent(const FConcertSessionContext& InEventContext, const FConcertSequencerPrecacheEvent& InEvent)
+void FConcertClientSequencePreloader::OnPreloadEvent(const FConcertSessionContext& InEventContext, const FConcertSequencerPreloadRequest& InEvent)
 {
-	for (const FString& SequenceObjectPathStr : InEvent.SequenceObjectPaths)
+	for (const FTopLevelAssetPath& SequenceObjectPath : InEvent.SequenceObjectPaths)
 	{
-		const FName SequenceObjectPath = *SequenceObjectPathStr;
-		const bool bShouldPrecache = InEvent.bShouldBePrecached;
-		UE_LOG(LogConcertSequencerSync, Verbose, TEXT("OnPrecacheEvent: %s should be %s the precache set"),
-			*SequenceObjectPathStr, bShouldPrecache ? TEXT("added to") : TEXT("removed from"));
-		if (bShouldPrecache)
+		const bool bShouldPreload = InEvent.bShouldBePreloaded;
+		UE_LOGFMT(LogConcertSequencerSync, Verbose, "OnPreloadEvent: {Path} should be {Action} the preload set",
+			SequenceObjectPath.ToString(), bShouldPreload ? TEXT("added to") : TEXT("removed from"));
+		if (bShouldPreload)
 		{
-			if (TObjectPtr<ULevelSequence>* ExistingSequence = PrecachedSequences.Find(SequenceObjectPath);
+			if (TObjectPtr<ULevelSequence>* ExistingSequence = PreloadedSequences.Find(SequenceObjectPath);
 				ExistingSequence && *ExistingSequence)
 			{
-				UE_LOG(LogConcertSequencerSync, Warning, TEXT("OnPrecacheEvent: %s already in precache set"), *SequenceObjectPath.ToString());
+				UE_LOGFMT(LogConcertSequencerSync, Warning, "OnPreloadEvent: {Path} already in preload set", SequenceObjectPath.ToString());
 				continue;
 			}
-
-			ULevelSequence* LoadedSequence = LoadObject<ULevelSequence>(nullptr, *SequenceObjectPathStr);
-			if (!LoadedSequence)
+			else if (PreloadPendingSequences.Contains(SequenceObjectPath))
 			{
-				UE_LOG(LogConcertSequencerSync, Warning, TEXT("OnPrecacheEvent: Failed to load sequence %s"), *SequenceObjectPath.ToString());
+				UE_LOGFMT(LogConcertSequencerSync, Warning, "OnPreloadEvent: {Path} already has outstanding async load", SequenceObjectPath.ToString());
 				continue;
 			}
 
-			PrecachedSequences.Add(SequenceObjectPath, LoadedSequence);
+			const FSoftObjectPath SoftSequenceObjectPath(SequenceObjectPath);
+			TSoftObjectPtr<ULevelSequence> SoftSequenceObject(SoftSequenceObjectPath);
+
+			if (ULevelSequence* AlreadyLoadedSequence = SoftSequenceObject.Get())
+			{
+				UE_LOGFMT(LogConcertSequencerSync, Verbose, "OnPreloadEvent: {Path} was already loaded", SequenceObjectPath.ToString());
+				PreloadedSequences.Add(SequenceObjectPath, AlreadyLoadedSequence);
+
+				if (TSharedPtr<IConcertClientSession> Session = WeakSession.Pin(); ensure(Session))
+				{
+					// Inform the server.
+					FConcertSequencerPreloadAssetStatusMap Response;
+					Response.Sequences.Add(SequenceObjectPath, EConcertSequencerPreloadStatus::Succeeded);
+					Session->SendCustomEvent(Response, Session->GetSessionServerEndpointId(), EConcertMessageFlags::ReliableOrdered);
+				}
+			}
+			else
+			{
+				UE_LOGFMT(LogConcertSequencerSync, Verbose, "OnPreloadEvent: Initiating async package load for {Path}", SequenceObjectPath.ToString());
+				PreloadPendingSequences.Add(SequenceObjectPath, SoftSequenceObject);
+				LoadPackageAsync(SoftSequenceObject.GetLongPackageName(), FLoadPackageAsyncDelegate::CreateLambda(
+					[WeakThis = AsWeak(),
+					WeakRequestSession = WeakSession,
+					SequenceObjectPath,
+					SoftSequenceObject]
+					(const FName& PackageName, UPackage* LoadedPackage, EAsyncLoadingResult::Type Result)
+					{
+						TSharedPtr<FConcertClientSequencePreloader> This = WeakThis.Pin();
+						if (!This)
+						{
+							UE_LOGFMT(LogConcertSequencerSync, Warning, "Discarding async load result for stale preloader");
+							return;
+						}
+
+						TSharedPtr<IConcertClientSession> Session = This->WeakSession.Pin();
+						if (!Session || Session != WeakRequestSession)
+						{
+							UE_LOGFMT(LogConcertSequencerSync, Warning, "Discarding async load result issued by mismatched session");
+							return;
+						}
+
+						if (!This->PreloadPendingSequences.Remove(SequenceObjectPath))
+						{
+							UE_LOGFMT(LogConcertSequencerSync, Warning, "Discarding async load result for sequence no longer pending");
+							return;
+						}
+
+						if (Result != EAsyncLoadingResult::Succeeded)
+						{
+							UE_LOGFMT(LogConcertSequencerSync, Error, "EAsyncLoadingResult != Succeeded for {Package} ({Result})", PackageName, Result);
+						}
+
+						ULevelSequence* LoadedSequence = SoftSequenceObject.Get();
+						if (!LoadedSequence)
+						{
+							UE_LOGFMT(LogConcertSequencerSync, Error, "Failed to resolve {Path} after async package load", SequenceObjectPath.ToString());
+						}
+
+						// Inform the server of the async load result, and GC ref the loaded sequence if successful.
+						FConcertSequencerPreloadAssetStatusMap Response;
+
+						if (Result == EAsyncLoadingResult::Succeeded && LoadedSequence)
+						{
+							UE_LOGFMT(LogConcertSequencerSync, Verbose, "OnPreloadEvent: Async load completed successfully for {Path}", SequenceObjectPath.ToString());
+							This->PreloadedSequences.Add(SequenceObjectPath, LoadedSequence);
+
+							Response.Sequences.Add(SequenceObjectPath, EConcertSequencerPreloadStatus::Succeeded);
+						}
+						else
+						{
+							Response.Sequences.Add(SequenceObjectPath, EConcertSequencerPreloadStatus::Failed);
+						}
+
+						Session->SendCustomEvent(Response, Session->GetSessionServerEndpointId(), EConcertMessageFlags::ReliableOrdered);
+					}
+				));
+			}
 		}
 		else
 		{
-			if (!PrecachedSequences.Remove(SequenceObjectPath))
+			if (!(PreloadPendingSequences.Remove(SequenceObjectPath)
+				|| PreloadedSequences.Remove(SequenceObjectPath)))
 			{
-				UE_LOG(LogConcertSequencerSync, Warning, TEXT("OnPrecacheEvent: %s not in precache set"), *SequenceObjectPath.ToString());
+				UE_LOGFMT(LogConcertSequencerSync, Warning, "OnPreloadEvent: Tried to remove {Path} not in preload set", *SequenceObjectPath.ToString());
 			}
 		}
 	}
@@ -1142,7 +1264,7 @@ void FConcertClientSequencerManager::OnWorkspaceEndFrameCompleted()
 void FConcertClientSequencerManager::AddReferencedObjects(FReferenceCollector& Collector)
 {
 	Collector.AddReferencedObjects(SequencePlayers);
-	Collector.AddReferencedObjects(PrecachedSequences);
+	Preloader->AddReferencedObjects(Collector);
 }
 
 #endif
diff --git a/Engine/Plugins/Developer/Concert/ConcertSync/ConcertSyncClient/Source/ConcertSyncClient/Private/ConcertClientSequencerManager.h b/Engine/Plugins/Developer/Concert/ConcertSync/ConcertSyncClient/Source/ConcertSyncClient/Private/ConcertClientSequencerManager.h
--- a/Engine/Plugins/Developer/Concert/ConcertSync/ConcertSyncClient/Source/ConcertSyncClient/Private/ConcertClientSequencerManager.h
+++ b/Engine/Plugins/Developer/Concert/ConcertSync/ConcertSyncClient/Source/ConcertSyncClient/Private/ConcertClientSequencerManager.h
@@ -17,6 +17,7 @@ class ALevelSequenceActor;
 #if WITH_EDITOR
 
 class ISequencer;
+class FConcertClientSequencePreloader;
 class FConcertClientWorkspace;
 struct FSequencerInitParams;
 class ULevelSequence;
@@ -199,13 +200,6 @@ private:
 	 */
 	void OnTimeAdjustmentEvent(const FConcertSessionContext& InEventContext, const FConcertSequencerTimeAdjustmentEvent& InEvent);
 
-	/**
-	 * Called on receipt of an external sequence precaching event from the server
-	 * @param InEventContext             The session context for the received message.
-	 * @param InEvent                    The sequence precache event.
-	 */
-	void OnPrecacheEvent(const FConcertSessionContext& InEventContext, const FConcertSequencerPrecacheEvent& InEvent);
-
 	/**
 	 * Called when the global time has been changed for the specified Sequencer
 	 *
@@ -362,8 +356,8 @@ private:
 	/** Map of opened sequence players, if not in editor mode. */
 	TMap<FName, TObjectPtr<ALevelSequenceActor>> SequencePlayers;
 
-	/** Map of precached sequence assets. */
-	TMap<FName, TObjectPtr<ULevelSequence>> PrecachedSequences;
+	/** Handle preload requests from the server, asynchronously loading and maintaining GC references to sequence assets. */
+	TSharedRef<FConcertClientSequencePreloader> Preloader;
 
 	/** Boolean that is set when we are handling any transport event to prevent re-entrancy */
 	bool bRespondingToTransportEvent;
diff --git a/Engine/Plugins/Developer/Concert/ConcertSync/ConcertSyncCore/Source/ConcertSyncCore/Public/ConcertSequencerMessages.h b/Engine/Plugins/Developer/Concert/ConcertSync/ConcertSyncCore/Source/ConcertSyncCore/Public/ConcertSequencerMessages.h
--- a/Engine/Plugins/Developer/Concert/ConcertSync/ConcertSyncCore/Source/ConcertSyncCore/Public/ConcertSequencerMessages.h
+++ b/Engine/Plugins/Developer/Concert/ConcertSync/ConcertSyncCore/Source/ConcertSyncCore/Public/ConcertSequencerMessages.h
@@ -131,7 +131,7 @@ struct FConcertSequencerTimeAdjustmentEvent
 
 /**
  * Event indicating one or more sequences have been added or removed from the
- * set of sequences to keep precached for quick dynamic instantiation.
+ * set of sequences to keep preloaded for quick dynamic instantiation.
  *
  * Can be sent by clients as a request to add or remove their references.
  *
@@ -139,17 +139,69 @@ struct FConcertSequencerTimeAdjustmentEvent
  * or as an initial snapshot of the complete set when joining a session.
  */
 USTRUCT()
-struct FConcertSequencerPrecacheEvent
+struct FConcertSequencerPreloadRequest
 {
 	GENERATED_BODY()
 
 	/** The list of full paths to affected sequences. */
 	UPROPERTY()
-	TArray<FString> SequenceObjectPaths;
+	TArray<FTopLevelAssetPath> SequenceObjectPaths;
 
-	/** True if being added to the precache set, false if being removed. */
+	/** True if being added to the preload set, false if being removed. */
 	UPROPERTY()
-	bool bShouldBePrecached = false;
+	bool bShouldBePreloaded = false;
+};
+
+UENUM()
+enum class EConcertSequencerPreloadStatus : uint8
+{
+	Pending,
+	Succeeded,
+	Failed,
+};
+
+/**
+ * Can be sent as an event by clients to indicate loading success/failure
+ * result of attempting to preload one or more sequence assets.
+ */
+USTRUCT()
+struct FConcertSequencerPreloadAssetStatusMap
+{
+	GENERATED_BODY()
+
+	UPROPERTY()
+	TMap<FTopLevelAssetPath, EConcertSequencerPreloadStatus> Sequences;
+};
+
+/**
+ * Sent as an event by server with preload status of one or more clients.
+ */
+USTRUCT()
+struct FConcertSequencerPreloadClientStatusMap
+{
+	GENERATED_BODY()
+
+	UPROPERTY()
+	TMap<FGuid, FConcertSequencerPreloadAssetStatusMap> ClientEndpoints;
+
+	/** Merge (add or replace) statuses for a specified client into this object. */
+	void UpdateFrom(const FGuid& Client, const FConcertSequencerPreloadAssetStatusMap& Updates)
+	{
+		FConcertSequencerPreloadAssetStatusMap& List = ClientEndpoints.FindOrAdd(Client);
+		for (const TPair<FTopLevelAssetPath, EConcertSequencerPreloadStatus>& Update : Updates.Sequences)
+		{
+			List.Sequences.FindOrAdd(Update.Key) = Update.Value;
+		}
+	}
+
+	/** Clear all clients' references to a specified asset. */
+	void Remove(const FTopLevelAssetPath& Sequence)
+	{
+		for (TPair<FGuid, FConcertSequencerPreloadAssetStatusMap>& Client : ClientEndpoints)
+		{
+			Client.Value.Sequences.Remove(Sequence);
+		}
+	}
 };
 
 
diff --git a/Engine/Plugins/Developer/Concert/ConcertSync/ConcertSyncServer/Source/ConcertSyncServer/Private/ConcertServerSequencerManager.cpp b/Engine/Plugins/Developer/Concert/ConcertSync/ConcertSyncServer/Source/ConcertSyncServer/Private/ConcertServerSequencerManager.cpp
--- a/Engine/Plugins/Developer/Concert/ConcertSync/ConcertSyncServer/Source/ConcertSyncServer/Private/ConcertServerSequencerManager.cpp
+++ b/Engine/Plugins/Developer/Concert/ConcertSync/ConcertSyncServer/Source/ConcertSyncServer/Private/ConcertServerSequencerManager.cpp
@@ -5,6 +5,8 @@
 #include "ConcertSyncServerLiveSession.h"
 #include "ConcertWorkspaceMessages.h"
 #include "ConcertLogGlobal.h"
+#include "Logging/StructuredLog.h"
+
 
 FConcertServerSequencerManager::FConcertServerSequencerManager(const TSharedRef<FConcertSyncServerLiveSession>& InLiveSession)
 {
@@ -29,7 +31,9 @@ void FConcertServerSequencerManager::BindSession(const TSharedRef<FConcertSyncSe
 	LiveSession->GetSession().RegisterCustomEventHandler<FConcertSequencerOpenEvent>(this, &FConcertServerSequencerManager::HandleSequencerOpenEvent);
 	LiveSession->GetSession().RegisterCustomEventHandler<FConcertSequencerTimeAdjustmentEvent>(this, &FConcertServerSequencerManager::HandleSequencerTimeAdjustmentEvent);
 	LiveSession->GetSession().RegisterCustomEventHandler<FConcertWorkspaceSyncAndFinalizeCompletedEvent>(this, &FConcertServerSequencerManager::HandleWorkspaceSyncAndFinalizeCompletedEvent);
-	LiveSession->GetSession().RegisterCustomEventHandler<FConcertSequencerPrecacheEvent>(this, &FConcertServerSequencerManager::HandleSequencerPrecacheEvent);
+
+	LiveSession->GetSession().RegisterCustomEventHandler<FConcertSequencerPreloadRequest>(this, &FConcertServerSequencerManager::HandleSequencerPreloadRequestEvent);
+	LiveSession->GetSession().RegisterCustomEventHandler<FConcertSequencerPreloadAssetStatusMap>(this, &FConcertServerSequencerManager::HandleSequencerPreloadStatusEvent);
 }
 
 void FConcertServerSequencerManager::UnbindSession()
@@ -42,7 +46,9 @@ void FConcertServerSequencerManager::UnbindSession()
 		LiveSession->GetSession().UnregisterCustomEventHandler<FConcertSequencerOpenEvent>(this);
 		LiveSession->GetSession().UnregisterCustomEventHandler<FConcertSequencerTimeAdjustmentEvent>(this);
 		LiveSession->GetSession().UnregisterCustomEventHandler<FConcertWorkspaceSyncAndFinalizeCompletedEvent>(this);
-		LiveSession->GetSession().UnregisterCustomEventHandler<FConcertSequencerPrecacheEvent>(this);
+
+		LiveSession->GetSession().UnregisterCustomEventHandler<FConcertSequencerPreloadRequest>(this);
+		LiveSession->GetSession().UnregisterCustomEventHandler<FConcertSequencerPreloadAssetStatusMap>(this);
 
 		LiveSession.Reset();
 	}
@@ -108,27 +114,27 @@ void FConcertServerSequencerManager::HandleSequencerCloseEvent(const FConcertSes
 	}
 }
 
-void FConcertServerSequencerManager::HandleSequencerPrecacheEvent(const FConcertSessionContext& InEventContext, const FConcertSequencerPrecacheEvent& InEvent)
+void FConcertServerSequencerManager::HandleSequencerPreloadRequestEvent(const FConcertSessionContext& InEventContext, const FConcertSequencerPreloadRequest& InEvent)
 {
 	const FGuid& RequestClient = InEventContext.SourceEndpointId;
-	const bool bClientWantsPrecached = InEvent.bShouldBePrecached;
+	const bool bClientWantsPreloaded = InEvent.bShouldBePreloaded;
 
-	UE_LOG(LogConcert, Verbose,
-		TEXT("FConcertServerSequencerManager: Precache request from client %s to %s %u sequences"),
+	UE_LOGFMT(LogConcert, Verbose,
+		"FConcertServerSequencerManager: Preload request from client {Client} to {AddOrRemove} {NumSequences} sequences",
 		*RequestClient.ToString(),
-		bClientWantsPrecached ? TEXT("add") : TEXT("remove"),
+		bClientWantsPreloaded ? TEXT("add") : TEXT("remove"),
 		InEvent.SequenceObjectPaths.Num());
 
 	// Represents the net result to broadcast, if any. Contains only sequences
 	// which gained their first, or lost their last, referencer.
-	FConcertSequencerPrecacheEvent OutChanges;
-	OutChanges.bShouldBePrecached = bClientWantsPrecached;
+	FConcertSequencerPreloadRequest OutChanges;
+	OutChanges.bShouldBePreloaded = bClientWantsPreloaded;
 
-	for (const FString& SequenceObjectPath : InEvent.SequenceObjectPaths)
+	for (const FTopLevelAssetPath& SequenceObjectPath : InEvent.SequenceObjectPaths)
 	{
-		if (bClientWantsPrecached)
+		if (bClientWantsPreloaded)
 		{
-			const bool bAddedFirstReferencer = AddSequencePrecacheForClient(RequestClient, SequenceObjectPath);
+			const bool bAddedFirstReferencer = AddSequencePreloadForClient(RequestClient, SequenceObjectPath);
 			if (bAddedFirstReferencer)
 			{
 				OutChanges.SequenceObjectPaths.Add(SequenceObjectPath);
@@ -136,28 +142,42 @@ void FConcertServerSequencerManager::HandleSequencerPrecacheEvent(const FConcert
 		}
 		else
 		{
-			const bool bRemovedLastReferencer = RemoveSequencePrecacheForClient(RequestClient, SequenceObjectPath);
+			const bool bRemovedLastReferencer = RemoveSequencePreloadForClient(RequestClient, SequenceObjectPath);
 			if (bRemovedLastReferencer)
 			{
 				OutChanges.SequenceObjectPaths.Add(SequenceObjectPath);
+				ClientPreloadStatuses.Remove(SequenceObjectPath);
 			}
 		}
 	}
 
 	if (OutChanges.SequenceObjectPaths.Num() > 0)
 	{
-		for (const FString& SequenceObjectPath : OutChanges.SequenceObjectPaths)
+		for (const FTopLevelAssetPath& SequenceObjectPath : OutChanges.SequenceObjectPaths)
 		{
-			UE_LOG(LogConcert, Verbose,
-				TEXT("FConcertServerSequencerManager: Sequence '%s' %s precache set"),
-				*SequenceObjectPath,
-				bClientWantsPrecached ? TEXT("added to") : TEXT("removed from"));
+			UE_LOGFMT(LogConcert, Verbose,
+				"FConcertServerSequencerManager: Sequence {Path} {AddedOrRemoved} preload set",
+				SequenceObjectPath.ToString(),
+				bClientWantsPreloaded ? TEXT("added to") : TEXT("removed from"));
 		}
 
 		LiveSession->GetSession().SendCustomEvent(OutChanges, LiveSession->GetSession().GetSessionClientEndpointIds(), EConcertMessageFlags::ReliableOrdered | EConcertMessageFlags::UniqueId);
 	}
 }
 
+void FConcertServerSequencerManager::HandleSequencerPreloadStatusEvent(const FConcertSessionContext& InEventContext, const FConcertSequencerPreloadAssetStatusMap& InEvent)
+{
+	// Update our cached status, and then forward the change to other clients.
+	const FGuid& Sender = InEventContext.SourceEndpointId;
+	ClientPreloadStatuses.UpdateFrom(Sender, InEvent);
+
+	FConcertSequencerPreloadClientStatusMap ForwardEvent;
+	ForwardEvent.UpdateFrom(Sender, InEvent);
+	TArray<FGuid> ForwardRecipients = LiveSession->GetSession().GetSessionClientEndpointIds();
+	ForwardRecipients.RemoveSingle(InEventContext.SourceEndpointId);
+	LiveSession->GetSession().SendCustomEvent(ForwardEvent, ForwardRecipients, EConcertMessageFlags::ReliableOrdered | EConcertMessageFlags::UniqueId);
+}
+
 void FConcertServerSequencerManager::HandleWorkspaceSyncAndFinalizeCompletedEvent(const FConcertSessionContext& InEventContext, const FConcertWorkspaceSyncAndFinalizeCompletedEvent& InEvent)
 {
 	FConcertSequencerStateSyncEvent SequencerStateSyncEvent;
@@ -172,6 +192,7 @@ void FConcertServerSequencerManager::HandleWorkspaceSyncAndFinalizeCompletedEven
 void FConcertServerSequencerManager::HandleSessionClientChanged(IConcertServerSession& InSession, EConcertClientStatus InClientStatus, const FConcertSessionClientInfo& InClientInfo)
 {
 	check(&InSession == &LiveSession->GetSession());
+
 	// Remove the client from all open sequences
 	if (InClientStatus == EConcertClientStatus::Disconnected)
 	{
@@ -192,64 +213,64 @@ void FConcertServerSequencerManager::HandleSessionClientChanged(IConcertServerSe
 		}
 	}
 
-	// Newly connected clients need to be sent the current set of precached sequences.
+	// Newly connected clients need to be sent the current set of preloaded sequences.
 	// Disconnecting clients need their references removed, which may update other clients.
 	if (InClientStatus == EConcertClientStatus::Connected ||
 		InClientStatus == EConcertClientStatus::Disconnected)
 	{
-		FConcertSequencerPrecacheEvent PrecacheEvent;
+		FConcertSequencerPreloadRequest PreloadRequest;
 		TArray<FGuid> EventRecipients;
 
 		if (InClientStatus == EConcertClientStatus::Connected)
 		{
-			PrecacheEvent.bShouldBePrecached = true;
+			PreloadRequest.bShouldBePreloaded = true;
 			EventRecipients.Add(InClientInfo.ClientEndpointId);
 		}
 		else
 		{
-			PrecacheEvent.bShouldBePrecached = false;
+			PreloadRequest.bShouldBePreloaded = false;
 			EventRecipients = LiveSession->GetSession().GetSessionClientEndpointIds();
 		}
 
-		for (TMap<FName, FPrecachingState>::TIterator It = PrecacheStates.CreateIterator(); It; ++It)
+		for (TMap<FTopLevelAssetPath, TSet<FGuid>>::TIterator It = PreloadRequesters.CreateIterator(); It; ++It)
 		{
 			if (InClientStatus == EConcertClientStatus::Connected)
 			{
-				ensure(It->Value.ReferencingClientEndpoints.Num() > 0);
-				PrecacheEvent.SequenceObjectPaths.Add(It->Key.ToString());
+				ensure(It->Value.Num() > 0);
+				PreloadRequest.SequenceObjectPaths.Add(It->Key);
 			}
 			else
 			{
-				if (It->Value.ReferencingClientEndpoints.Remove(InClientInfo.ClientEndpointId))
+				if (It->Value.Remove(InClientInfo.ClientEndpointId))
 				{
-					if (It->Value.ReferencingClientEndpoints.Num() == 0)
+					if (It->Value.Num() == 0)
 					{
-						PrecacheEvent.SequenceObjectPaths.Add(It->Key.ToString());
+						PreloadRequest.SequenceObjectPaths.Add(It->Key);
 						It.RemoveCurrent();
 					}
 				}
 			}
 		}
 
-		if (PrecacheEvent.SequenceObjectPaths.Num() > 0)
+		if (PreloadRequest.SequenceObjectPaths.Num() > 0)
 		{
-			for (const FString& SequenceObjectPath : PrecacheEvent.SequenceObjectPaths)
+			for (const FTopLevelAssetPath& SequenceObjectPath : PreloadRequest.SequenceObjectPaths)
 			{
 				if (InClientStatus == EConcertClientStatus::Connected)
 				{
-					UE_LOG(LogConcert, Verbose,
-						TEXT("FConcertServerSequencerManager: Client connected; notifying precache set contains sequence '%s'"),
-						*SequenceObjectPath);
+					UE_LOGFMT(LogConcert, Verbose,
+						"FConcertServerSequencerManager: Client connected; notifying preload set contains sequence '{Path}'",
+						SequenceObjectPath.ToString());
 				}
 				else
 				{
-					UE_LOG(LogConcert, Verbose,
-						TEXT("FConcertServerSequencerManager: Client disconnected; last reference to '%s' was released, removed from precache set"),
-						*SequenceObjectPath);
+					UE_LOGFMT(LogConcert, Verbose,
+						"FConcertServerSequencerManager: Client disconnected; last reference to '{Path}' was released, removed from preload set",
+						SequenceObjectPath.ToString());
 				}
 			}
 
-			LiveSession->GetSession().SendCustomEvent(PrecacheEvent, EventRecipients, EConcertMessageFlags::ReliableOrdered | EConcertMessageFlags::UniqueId);
+			LiveSession->GetSession().SendCustomEvent(PreloadRequest, EventRecipients, EConcertMessageFlags::ReliableOrdered | EConcertMessageFlags::UniqueId);
 		}
 	}
 
@@ -259,48 +280,48 @@ void FConcertServerSequencerManager::HandleSessionClientChanged(IConcertServerSe
 	// stream.
 }
 
-bool FConcertServerSequencerManager::AddSequencePrecacheForClient(const FGuid& RequestClient, const FString& SequenceObjectPath)
+bool FConcertServerSequencerManager::AddSequencePreloadForClient(const FGuid& RequestClient, const FTopLevelAssetPath& SequenceObjectPath)
 {
 	bool bAddedFirstReferencer = false;
 
-	FPrecachingState& State = PrecacheStates.FindOrAdd(*SequenceObjectPath);
+	TSet<FGuid>& Requesters = PreloadRequesters.FindOrAdd(SequenceObjectPath);
 	bool bAlreadyInSet = false;
-	State.ReferencingClientEndpoints.Add(RequestClient, &bAlreadyInSet);
+	Requesters.Add(RequestClient, &bAlreadyInSet);
 	if (!bAlreadyInSet)
 	{
-		if (State.ReferencingClientEndpoints.Num() == 1)
+		if (Requesters.Num() == 1)
 		{
 			bAddedFirstReferencer = true;
 		}
 	}
 	else
 	{
-		UE_LOG(LogConcert, Warning, TEXT("FConcertServerSequencerManager: Client %s requested redundant add precache for sequence %s"),
-			*RequestClient.ToString(), *SequenceObjectPath);
+		UE_LOGFMT(LogConcert, Warning, "FConcertServerSequencerManager: Client {Client} requested redundant add preload for sequence {Path}",
+			RequestClient.ToString(), SequenceObjectPath.ToString());
 	}
 
 	return bAddedFirstReferencer;
 }
 
-bool FConcertServerSequencerManager::RemoveSequencePrecacheForClient(const FGuid& RequestClient, const FString& SequenceObjectPath)
+bool FConcertServerSequencerManager::RemoveSequencePreloadForClient(const FGuid& RequestClient, const FTopLevelAssetPath& SequenceObjectPath)
 {
 	bool bRemovedLastReferencer = false;
 
-	FPrecachingState* MaybeState = PrecacheStates.Find(*SequenceObjectPath);
-	if (MaybeState && MaybeState->ReferencingClientEndpoints.Remove(RequestClient))
+	TSet<FGuid>* MaybeRequesters = PreloadRequesters.Find(SequenceObjectPath);
+	if (MaybeRequesters && MaybeRequesters->Remove(RequestClient))
 	{
-		if (MaybeState->ReferencingClientEndpoints.Num() == 0)
+		if (MaybeRequesters->Num() == 0)
 		{
 			// Removed last reference.
-			MaybeState = nullptr;
-			PrecacheStates.Remove(*SequenceObjectPath);
+			MaybeRequesters = nullptr;
+			PreloadRequesters.Remove(SequenceObjectPath);
 			bRemovedLastReferencer = true;
 		}
 	}
 	else
 	{
-		UE_LOG(LogConcert, Warning, TEXT("FConcertServerSequencerManager: Client %s attempted invalid release precache for sequence %s"),
-			*RequestClient.ToString(), *SequenceObjectPath);
+		UE_LOGFMT(LogConcert, Warning, "FConcertServerSequencerManager: Client {Client} attempted invalid release preload for sequence {Path}",
+			RequestClient.ToString(), SequenceObjectPath.ToString());
 	}
 
 	return bRemovedLastReferencer;
diff --git a/Engine/Plugins/Developer/Concert/ConcertSync/ConcertSyncServer/Source/ConcertSyncServer/Private/ConcertServerSequencerManager.h b/Engine/Plugins/Developer/Concert/ConcertSync/ConcertSyncServer/Source/ConcertSyncServer/Private/ConcertServerSequencerManager.h
--- a/Engine/Plugins/Developer/Concert/ConcertSync/ConcertSyncServer/Source/ConcertSyncServer/Private/ConcertServerSequencerManager.h
+++ b/Engine/Plugins/Developer/Concert/ConcertSync/ConcertSyncServer/Source/ConcertSyncServer/Private/ConcertServerSequencerManager.h
@@ -36,8 +36,11 @@ private:
 	/** Handler for the sequencer time adjustment event. */
 	void HandleSequencerTimeAdjustmentEvent(const FConcertSessionContext& InEventContext, const FConcertSequencerTimeAdjustmentEvent& InEvent);
 
-	/** Handler for sequence precaching requests. */
-	void HandleSequencerPrecacheEvent(const FConcertSessionContext& InEventContext, const FConcertSequencerPrecacheEvent& InEvent);
+	/** Handler for sequence preloading requests. */
+	void HandleSequencerPreloadRequestEvent(const FConcertSessionContext& InEventContext, const FConcertSequencerPreloadRequest& InEvent);
+
+	/** Handler for sequence preloading status updates. */
+	void HandleSequencerPreloadStatusEvent(const FConcertSessionContext& InEventContext, const FConcertSequencerPreloadAssetStatusMap& InEvent);
 
 	/** Handler for the workspace sync and finalize completed event. */
 	void HandleWorkspaceSyncAndFinalizeCompletedEvent(const FConcertSessionContext& InEventContext, const FConcertWorkspaceSyncAndFinalizeCompletedEvent& InEvent);
@@ -47,10 +50,10 @@ private:
 
 
 	/** Adds the specified client endpoint as a referencer for the specified sequence, and returns true if that was the first reference for that sequence. */
-	bool AddSequencePrecacheForClient(const FGuid& RequestClient, const FString& SequenceObjectPath);
+	bool AddSequencePreloadForClient(const FGuid& RequestClient, const FTopLevelAssetPath& SequenceObjectPath);
 
 	/** Adds the specified client endpoint as a referencer for the specified sequence, and returns true if that was the last remaining reference for that sequence. */
-	bool RemoveSequencePrecacheForClient(const FGuid& RequestClient, const FString& SequenceObjectPath);
+	bool RemoveSequencePreloadForClient(const FGuid& RequestClient, const FTopLevelAssetPath& SequenceObjectPath);
 
 
 	struct FConcertOpenSequencerState
@@ -68,12 +71,9 @@ private:
 	/** Live session tracked by this manager. */
 	TSharedPtr<FConcertSyncServerLiveSession> LiveSession;
 
-	struct FPrecachingState
-	{
-		/** Client endpoints which have requested this sequence be precached. */
-		TSet<FGuid> ReferencingClientEndpoints;
-	};
+	/** Map of level sequences indicating which client endpoints have requested preloading for them. */
+	TMap<FTopLevelAssetPath, TSet<FGuid>> PreloadRequesters;
 
-	/** Map of all level sequences for which precaching has been requested. */
-	TMap<FName, FPrecachingState> PrecacheStates;
+	/** Map of clients' individual progress preloading specific sequences. */
+	FConcertSequencerPreloadClientStatusMap ClientPreloadStatuses;
 };
diff --git a/Engine/Plugins/VirtualProduction/SequencerPlaylists/Source/SequencerPlaylists/Private/SequencerPlaylistItem_Sequence.cpp b/Engine/Plugins/VirtualProduction/SequencerPlaylists/Source/SequencerPlaylists/Private/SequencerPlaylistItem_Sequence.cpp
--- a/Engine/Plugins/VirtualProduction/SequencerPlaylists/Source/SequencerPlaylists/Private/SequencerPlaylistItem_Sequence.cpp
+++ b/Engine/Plugins/VirtualProduction/SequencerPlaylists/Source/SequencerPlaylists/Private/SequencerPlaylistItem_Sequence.cpp
@@ -96,7 +96,7 @@ void USequencerPlaylistItem_Sequence::PostEditChangeProperty(FPropertyChangedEve
 			GEditor->GetEditorSubsystem<USequencerPlaylistsSubsystem>();
 		if (Subsystem)
 		{
-			Subsystem->UpdatePrecacheSet();
+			Subsystem->UpdatePreloadSet();
 		}
 	}
 }
@@ -108,7 +108,7 @@ void USequencerPlaylistItem_Sequence::SetSequence(ULevelSequence* NewSequence)
 
 	if (USequencerPlaylistsSubsystem* Subsystem = GEditor->GetEditorSubsystem<USequencerPlaylistsSubsystem>())
 	{
-		Subsystem->UpdatePrecacheSet();
+		Subsystem->UpdatePreloadSet();
 	}
 }
 
diff --git a/Engine/Plugins/VirtualProduction/SequencerPlaylists/Source/SequencerPlaylists/Private/SequencerPlaylistsSubsystem.cpp b/Engine/Plugins/VirtualProduction/SequencerPlaylists/Source/SequencerPlaylists/Private/SequencerPlaylistsSubsystem.cpp
--- a/Engine/Plugins/VirtualProduction/SequencerPlaylists/Source/SequencerPlaylists/Private/SequencerPlaylistsSubsystem.cpp
+++ b/Engine/Plugins/VirtualProduction/SequencerPlaylists/Source/SequencerPlaylists/Private/SequencerPlaylistsSubsystem.cpp
@@ -7,6 +7,7 @@
 #include "IConcertSyncClientModule.h"
 #include "IConcertSession.h"
 #include "LevelSequence.h"
+#include "Logging/StructuredLog.h"
 #include "SequencerPlaylist.h"
 #include "SequencerPlaylistItem_Sequence.h"
 #include "SequencerPlaylistPlayer.h"
@@ -14,6 +15,9 @@
 #include "UObject/Package.h"
 
 
+#define LOCTEXT_NAMESPACE "SequencerPlaylists"
+
+
 void USequencerPlaylistsSubsystem::Initialize(FSubsystemCollectionBase& Collection)
 {
 	Super::Initialize(Collection);
@@ -85,17 +89,130 @@ void USequencerPlaylistsSubsystem::NotifyEditorClosed(SSequencerPlaylistPanel* E
 	ensure(EditorPlayers.Remove(EditorHandle));
 	ensure(EditorPackages.Remove(EditorHandle));
 
-	UpdatePrecacheSet();
+	UpdatePreloadSet();
+}
+
+
+bool USequencerPlaylistsSubsystem::IsPreloadingActive() const
+{
+	// If we're in a multi-user session, we are responding to remote preloading requests.
+	return WeakSession.IsValid();
+}
+
+
+TPair<EConcertSequencerPreloadStatus, FText> USequencerPlaylistsSubsystem::GetPreloadStatusForSequence(const FTopLevelAssetPath& SequencePath)
+{
+	TPair<EConcertSequencerPreloadStatus, FText> Result;
+
+	TSharedPtr<IConcertClientSession> Session = WeakSession.Pin();
+	if (!ensure(Session))
+	{
+		Result.Key = EConcertSequencerPreloadStatus::Failed;
+		Result.Value = LOCTEXT("PreloadStatusNoSession", "Not in multi-user session");
+		return Result;
+	}
+
+	const TArray<FConcertSessionClientInfo> AllClients = Session->GetSessionClients();
+	TSet<FGuid> PendingClients, FailedClients;
+	TMap<FGuid, FConcertClientInfo> ClientGuidToInfo;
+	for (const FConcertSessionClientInfo& Client : AllClients)
+	{
+		ClientGuidToInfo.Add(Client.ClientEndpointId, Client.ClientInfo);
+
+		const FClientSequenceKey CompositeKey{ Client.ClientEndpointId, SequencePath };
+		EConcertSequencerPreloadStatus Status = EConcertSequencerPreloadStatus::Pending;
+		if (const EConcertSequencerPreloadStatus* MaybeStatus = PreloadStates.Find(CompositeKey))
+		{
+			Status = *MaybeStatus;
+		}
+
+		switch (Status)
+		{
+		case EConcertSequencerPreloadStatus::Pending:
+			PendingClients.Add(Client.ClientEndpointId);
+			break;
+		case EConcertSequencerPreloadStatus::Failed:
+			FailedClients.Add(Client.ClientEndpointId);
+			break;
+		}
+	}
+
+	// Determine overall status enum
+	if (FailedClients.Num() > 0)
+	{
+		Result.Key = EConcertSequencerPreloadStatus::Failed;
+	}
+	else if (PendingClients.Num() > 0)
+	{
+		Result.Key = EConcertSequencerPreloadStatus::Pending;
+	}
+	else
+	{
+		Result.Key = EConcertSequencerPreloadStatus::Succeeded;
+	}
+
+	// Construct accompanying textual detail
+	const FText PreloadClientDelimiter = LOCTEXT("PreloadStatusClientDelimiter", ", ");
+	TArray<FText> StatusGroups;
+	if (FailedClients.Num())
+	{
+		TArray<FText> ClientDisplayNames;
+		for (const FGuid& ClientId : FailedClients)
+		{
+			ClientDisplayNames.Add(FText::FromString(ClientGuidToInfo[ClientId].DisplayName));
+		}
+		StatusGroups.Add(FText::Format(
+			LOCTEXT("PreloadStatusFailedList", "Failed: {0}"),
+			FText::Join(PreloadClientDelimiter, ClientDisplayNames)));
+	}
+
+	if (PendingClients.Num())
+	{
+		TArray<FText> ClientDisplayNames;
+		for (const FGuid& ClientId : PendingClients)
+		{
+			ClientDisplayNames.Add(FText::FromString(ClientGuidToInfo[ClientId].DisplayName));
+		}
+		StatusGroups.Add(FText::Format(
+			LOCTEXT("PreloadStatusPendingList", "Pending: {0}"),
+			FText::Join(PreloadClientDelimiter, ClientDisplayNames)));
+	}
+
+	if (FailedClients.Num() == 0 && PendingClients.Num() == 0)
+	{
+		StatusGroups.Add(LOCTEXT("PreloadStatusAllSuccessful", "All clients preloaded successfully."));
+	}
+	else
+	{
+		TArray<FText> ClientDisplayNames;
+		for (const FConcertSessionClientInfo& Client : AllClients)
+		{
+			if (!FailedClients.Contains(Client.ClientEndpointId) && !PendingClients.Contains(Client.ClientEndpointId))
+			{
+				ClientDisplayNames.Add(FText::FromString(Client.ClientInfo.DisplayName));
+			}
+		}
+
+		if (ClientDisplayNames.Num())
+		{
+			StatusGroups.Add(FText::Format(
+				LOCTEXT("PreloadStatusCompleteList", "Complete: {0}"),
+				FText::Join(PreloadClientDelimiter, ClientDisplayNames)));
+		}
+	}
+
+	Result.Value = FText::Join(LOCTEXT("PreloadGroupsDelimiter", "\n"), StatusGroups);
+	return Result;
 }
 
 
-void USequencerPlaylistsSubsystem::UpdatePrecacheSet()
+void USequencerPlaylistsSubsystem::UpdatePreloadSet()
 {
-	// Determine which sequences (if any) were added to/removed from the precache set.
+	// Determine which sequences (if any) were added to/removed from the preload set.
 	// We copy the previous set, recompute the current set, then compare them.
-	const TSet<TObjectPtr<ULevelSequence>> PreviousPrecache = PrecacheSequences;
+	const TSet<TObjectPtr<ULevelSequence>> PreviousPreload = PreloadSequences;
 
-	PrecacheSequences.Empty(PrecacheSequences.Num());
+	PreloadSequences.Empty(PreloadSequences.Num());
 
 	using FPlaylistPair = TPair<FSequencerPlaylistEditorHandle, TObjectPtr<USequencerPlaylist>>;
 	for (const FPlaylistPair& EditorPlaylist : EditorPlaylists)
@@ -105,40 +222,40 @@ void USequencerPlaylistsSubsystem::UpdatePrecacheSet()
 			USequencerPlaylistItem_Sequence* SequenceItem = Cast<USequencerPlaylistItem_Sequence>(Item);
 			if (SequenceItem && SequenceItem->GetSequence())
 			{
-				PrecacheSequences.Add(SequenceItem->GetSequence());
+				PreloadSequences.Add(SequenceItem->GetSequence());
 			}
 		}
 	}
 
-	FConcertSequencerPrecacheEvent PrecacheAddedEvent{ .bShouldBePrecached = true };
-	FConcertSequencerPrecacheEvent PrecacheRemovedEvent{ .bShouldBePrecached = false };
+	FConcertSequencerPreloadRequest PreloadAddedEvent{ .bShouldBePreloaded = true };
+	FConcertSequencerPreloadRequest PreloadRemovedEvent{ .bShouldBePreloaded = false };
 
-	const TSet<TObjectPtr<ULevelSequence>> PrecacheUnion = PreviousPrecache.Union(PrecacheSequences);
-	for (const TObjectPtr<ULevelSequence>& Sequence : PrecacheUnion)
+	const TSet<TObjectPtr<ULevelSequence>> PreloadUnion = PreviousPreload.Union(PreloadSequences);
+	for (const TObjectPtr<ULevelSequence>& Sequence : PreloadUnion)
 	{
-		const FString SequenceObjectPath = Sequence->GetPathName();
-		if (!PreviousPrecache.Contains(Sequence))
+		const FTopLevelAssetPath SequenceObjectPath(Sequence);
+		if (!PreviousPreload.Contains(Sequence))
 		{
-			UE_LOG(LogSequencerPlaylists, Verbose, TEXT("Adding sequence '%s' to precache set"), *SequenceObjectPath);
-			PrecacheAddedEvent.SequenceObjectPaths.Add(*SequenceObjectPath);
+			UE_LOGFMT(LogSequencerPlaylists, Verbose, "Adding sequence '{Asset}' to preload set", UE::FAssetLog(Sequence));
+			PreloadAddedEvent.SequenceObjectPaths.Add(SequenceObjectPath);
 		}
-		else if (!PrecacheSequences.Contains(Sequence))
+		else if (!PreloadSequences.Contains(Sequence))
 		{
-			UE_LOG(LogSequencerPlaylists, Verbose, TEXT("Removing sequence '%s' from precache set"), *SequenceObjectPath);
-			PrecacheRemovedEvent.SequenceObjectPaths.Add(*SequenceObjectPath);
+			UE_LOGFMT(LogSequencerPlaylists, Verbose, "Removing sequence '{Asset}' from preload set", UE::FAssetLog(Sequence));
+			PreloadRemovedEvent.SequenceObjectPaths.Add(SequenceObjectPath);
 		}
 	}
 
 	if (TSharedPtr<IConcertClientSession> Session = WeakSession.Pin())
 	{
-		if (PrecacheAddedEvent.SequenceObjectPaths.Num() > 0)
+		if (PreloadAddedEvent.SequenceObjectPaths.Num() > 0)
 		{
-			Session->SendCustomEvent(PrecacheAddedEvent, Session->GetSessionServerEndpointId(), EConcertMessageFlags::ReliableOrdered);
+			Session->SendCustomEvent(PreloadAddedEvent, Session->GetSessionServerEndpointId(), EConcertMessageFlags::ReliableOrdered);
 		}
 
-		if (PrecacheRemovedEvent.SequenceObjectPaths.Num() > 0)
+		if (PreloadRemovedEvent.SequenceObjectPaths.Num() > 0)
 		{
-			Session->SendCustomEvent(PrecacheRemovedEvent, Session->GetSessionServerEndpointId(), EConcertMessageFlags::ReliableOrdered);
+			Session->SendCustomEvent(PreloadRemovedEvent, Session->GetSessionServerEndpointId(), EConcertMessageFlags::ReliableOrdered);
 		}
 	}
 }
@@ -148,16 +265,20 @@ void USequencerPlaylistsSubsystem::OnConcertSessionStartup(TSharedRef<IConcertCl
 {
 	WeakSession = InSession;
 
-	FConcertSequencerPrecacheEvent PrecacheEvent{ .bShouldBePrecached = true };
+	InSession->OnSessionClientChanged().AddUObject(this, &USequencerPlaylistsSubsystem::HandleSessionClientChanged);
+	InSession->RegisterCustomEventHandler<FConcertSequencerPreloadClientStatusMap>(this, &USequencerPlaylistsSubsystem::HandleSequencerPreloadStatusEvent);
 
-	for (const TObjectPtr<ULevelSequence>& Sequence : PrecacheSequences)
+	// Notify the server of any sequences we already wanted preloaded prior to joining the session.
+	FConcertSequencerPreloadRequest PreloadEvent{ .bShouldBePreloaded = true };
+
+	for (const TObjectPtr<ULevelSequence>& Sequence : PreloadSequences)
 	{
-		PrecacheEvent.SequenceObjectPaths.Add(*Sequence->GetPathName());
+		PreloadEvent.SequenceObjectPaths.Add(FTopLevelAssetPath(Sequence));
 	}
 
-	if (PrecacheEvent.SequenceObjectPaths.Num() > 0)
+	if (PreloadEvent.SequenceObjectPaths.Num() > 0)
 	{
-		InSession->SendCustomEvent(PrecacheEvent, InSession->GetSessionServerEndpointId(), EConcertMessageFlags::ReliableOrdered);
+		InSession->SendCustomEvent(PreloadEvent, InSession->GetSessionServerEndpointId(), EConcertMessageFlags::ReliableOrdered);
 	}
 }
 
@@ -167,3 +288,41 @@ void USequencerPlaylistsSubsystem::OnConcertSessionShutdown(TSharedRef<IConcertC
 	ensure(WeakSession == InSession);
 	WeakSession = nullptr;
 }
+
+
+void USequencerPlaylistsSubsystem::HandleSessionClientChanged(IConcertClientSession& InSession, EConcertClientStatus InClientStatus, const FConcertSessionClientInfo& InClientInfo)
+{
+	// Clear any PreloadStates belonging to a disconnecting client.
+	if (InClientStatus == EConcertClientStatus::Disconnected)
+	{
+		const FGuid& DisconnectingClient = InClientInfo.ClientEndpointId;
+		for (TMap<FClientSequenceKey, EConcertSequencerPreloadStatus>::TIterator Iter = PreloadStates.CreateIterator(); Iter; ++Iter)
+		{
+			const FClientSequenceKey& ClientSeqPair = Iter->Key;
+			if (ClientSeqPair.Key == DisconnectingClient)
+			{
+				Iter.RemoveCurrent();
+			}
+		}
+	}
+}
+
+
+void USequencerPlaylistsSubsystem::HandleSequencerPreloadStatusEvent(const FConcertSessionContext& InEventContext, const FConcertSequencerPreloadClientStatusMap& InEvent)
+{
+	for (const TPair<FGuid, FConcertSequencerPreloadAssetStatusMap>& EndpointStatus : InEvent.ClientEndpoints)
+	{
+		const FGuid& Endpoint = EndpointStatus.Key;
+		for (const TPair<FTopLevelAssetPath, EConcertSequencerPreloadStatus>& Status : EndpointStatus.Value.Sequences)
+		{
+			const FClientSequenceKey CacheKey{ Endpoint, Status.Key };
+			PreloadStates.FindOrAdd(CacheKey) = Status.Value;
+
+			UE_LOGFMT(LogSequencerPlaylists, Verbose, "Remote endpoint {Endpoint} reported {Result} for preload of '{Asset}'",
+				Endpoint.ToString(), static_cast<int32>(Status.Value), Status.Key.ToString() );
+		}
+	}
+}
+
+
+#undef LOCTEXT_NAMESPACE
diff --git a/Engine/Plugins/VirtualProduction/SequencerPlaylists/Source/SequencerPlaylists/Private/SequencerPlaylistsSubsystem.h b/Engine/Plugins/VirtualProduction/SequencerPlaylists/Source/SequencerPlaylists/Private/SequencerPlaylistsSubsystem.h
--- a/Engine/Plugins/VirtualProduction/SequencerPlaylists/Source/SequencerPlaylists/Private/SequencerPlaylistsSubsystem.h
+++ b/Engine/Plugins/VirtualProduction/SequencerPlaylists/Source/SequencerPlaylists/Private/SequencerPlaylistsSubsystem.h
@@ -6,6 +6,11 @@
 #include "SequencerPlaylistsSubsystem.generated.h"
 
 
+enum class EConcertClientStatus : uint8;
+enum class EConcertSequencerPreloadStatus : uint8;
+struct FConcertSessionClientInfo;
+struct FConcertSessionContext;
+struct FConcertSequencerPreloadClientStatusMap;
 class IConcertClientSession;
 class SSequencerPlaylistPanel;
 class ULevelSequence;
@@ -54,7 +59,14 @@ public:
 	USequencerPlaylistPlayer* CreatePlayerForEditor(TSharedRef<SSequencerPlaylistPanel> Editor);
 	void NotifyEditorClosed(SSequencerPlaylistPanel* Editor);
 
-	void UpdatePrecacheSet();
+	void UpdatePreloadSet();
+
+	bool IsPreloadingActive() const;
+
+	/** If any client has failed or is pending, the enum will indicate failure or pending respectively (otherwise, success).
+	 * The accompanying text contains a (potentially multi-line) summary of which multi-user clients are in which states.
+	 */
+	TPair<EConcertSequencerPreloadStatus, FText> GetPreloadStatusForSequence(const FTopLevelAssetPath& SequencePath);
 
 protected:
 	USequencerPlaylist* CreateTransientPlaylistForEditor(TSharedRef<SSequencerPlaylistPanel> Editor);
@@ -62,6 +74,9 @@ protected:
 	void OnConcertSessionStartup(TSharedRef<IConcertClientSession> InSession);
 	void OnConcertSessionShutdown(TSharedRef<IConcertClientSession> InSession);
 
+	void HandleSessionClientChanged(IConcertClientSession& InSession, EConcertClientStatus InClientStatus, const FConcertSessionClientInfo& InClientInfo);
+	void HandleSequencerPreloadStatusEvent(const FConcertSessionContext& InEventContext, const FConcertSequencerPreloadClientStatusMap& InEvent);
+
 protected:
 	UPROPERTY(Transient)
 	TMap<FSequencerPlaylistEditorHandle, TObjectPtr<UPackage>> EditorPackages;
@@ -73,7 +88,10 @@ protected:
 	TMap<FSequencerPlaylistEditorHandle, TObjectPtr<USequencerPlaylist>> EditorPlaylists;
 
 	UPROPERTY(Transient)
-	TSet<TObjectPtr<ULevelSequence>> PrecacheSequences;
+	TSet<TObjectPtr<ULevelSequence>> PreloadSequences;
+
+	using FClientSequenceKey = TPair<FGuid, FTopLevelAssetPath>;
+	TMap<FClientSequenceKey, EConcertSequencerPreloadStatus> PreloadStates;
 
 	/** Weak pointer to the client session with which to send events. May be null or stale. */
 	TWeakPtr<IConcertClientSession> WeakSession;
diff --git a/Engine/Plugins/VirtualProduction/SequencerPlaylists/Source/SequencerPlaylists/Private/SequencerPlaylistsWidgets.cpp b/Engine/Plugins/VirtualProduction/SequencerPlaylists/Source/SequencerPlaylists/Private/SequencerPlaylistsWidgets.cpp
--- a/Engine/Plugins/VirtualProduction/SequencerPlaylists/Source/SequencerPlaylists/Private/SequencerPlaylistsWidgets.cpp
+++ b/Engine/Plugins/VirtualProduction/SequencerPlaylists/Source/SequencerPlaylists/Private/SequencerPlaylistsWidgets.cpp
@@ -17,9 +17,13 @@
 
 #include "AssetRegistry/AssetRegistryModule.h"
 #include "AssetToolsModule.h"
+#include "ConcertSequencerMessages.h"
 #include "ContentBrowserModule.h"
+#include "EditorFontGlyphs.h"
 #include "FileHelpers.h"
 #include "Framework/MultiBox/MultiBoxBuilder.h"
+#include "IConcertSyncClient.h"
+#include "IConcertSyncClientModule.h"
 #include "IContentBrowserSingleton.h"
 #include "ISinglePropertyView.h"
 #include "Layout/WidgetPath.h"
@@ -50,6 +54,7 @@ const FName SSequencerPlaylistPanel::ColumnName_Items(TEXT("Items"));
 const FName SSequencerPlaylistPanel::ColumnName_Offset(TEXT("Offset"));
 const FName SSequencerPlaylistPanel::ColumnName_Pause(TEXT("Pause"));
 const FName SSequencerPlaylistPanel::ColumnName_Loop(TEXT("Loop"));
+const FName SSequencerPlaylistPanel::ColumnName_Preload(TEXT("Preload"));
 const FName SSequencerPlaylistPanel::ColumnName_HoverDetails(TEXT("HoverDetails"));
 
 
@@ -161,12 +166,26 @@ void SSequencerPlaylistPanel::Construct(const FArguments& InArgs, TSharedPtr<SDo
 			]
 		]
 	];
+
+	if (TSharedPtr<IConcertSyncClient> ConcertSyncClient = IConcertSyncClientModule::Get().GetClient(TEXT("MultiUser")))
+	{
+		IConcertClientRef ConcertClient = ConcertSyncClient->GetConcertClient();
+		ConcertClient->OnSessionStartup().AddRaw(this, &SSequencerPlaylistPanel::OnConcertSessionStartup);
+		ConcertClient->OnSessionShutdown().AddRaw(this, &SSequencerPlaylistPanel::OnConcertSessionShutdown);
+	}
 }
 
 
 SSequencerPlaylistPanel::~SSequencerPlaylistPanel()
 {
 	GEditor->GetEditorSubsystem<USequencerPlaylistsSubsystem>()->NotifyEditorClosed(this);
+
+	if (TSharedPtr<IConcertSyncClient> ConcertSyncClient = IConcertSyncClientModule::Get().GetClient(TEXT("MultiUser")))
+	{
+		IConcertClientRef ConcertClient = ConcertSyncClient->GetConcertClient();
+		ConcertClient->OnSessionStartup().RemoveAll(this);
+		ConcertClient->OnSessionShutdown().RemoveAll(this);
+	}
 }
 
 
@@ -472,10 +491,15 @@ TSharedRef<SWidget> SSequencerPlaylistPanel::Construct_ItemListView()
 					.OnAcceptDrop(this, &SSequencerPlaylistPanel::HandleAcceptDrop);
 			})
 		.HeaderRow(
-			SNew(SHeaderRow)
+			SAssignNew(HeaderRow, SHeaderRow)
 			+ SHeaderRow::Column(ColumnName_HoverTransport)
 				.DefaultLabel(FText::GetEmpty())
 				.FillSized(60.0f)
+			+ SHeaderRow::Column(ColumnName_Preload)
+				.DefaultLabel(LOCTEXT("ColumnLabelPreload", "Preload"))
+				.FillSized(60.0f)
+				.HAlignCell(HAlign_Center)
+				.VAlignCell(VAlign_Center)
 			+ SHeaderRow::Column(ColumnName_Items)
 				.DefaultLabel(LOCTEXT("ColumnLabelItems", "Playlist Items"))
 				.FillWidth(1.0f)
@@ -495,6 +519,12 @@ TSharedRef<SWidget> SSequencerPlaylistPanel::Construct_ItemListView()
 				.HAlignCell(HAlign_Center)
 				.VAlignCell(VAlign_Center)
 		);
+
+	if (ensure(HeaderRow))
+	{
+		const bool bPreloadVisible = IConcertSyncClientModule::Get().GetClient(TEXT("MultiUser")).IsValid();
+		HeaderRow->SetShowGeneratedColumn(ColumnName_Preload, bPreloadVisible);
+	}
 }
 
 
@@ -768,7 +798,7 @@ void SSequencerPlaylistPanel::LoadPlaylist(USequencerPlaylist* PlaylistToLoad)
 
 	if (USequencerPlaylistsSubsystem* Subsystem = GEditor->GetEditorSubsystem<USequencerPlaylistsSubsystem>())
 	{
-		Subsystem->UpdatePrecacheSet();
+		Subsystem->UpdatePreloadSet();
 	}
 }
 
@@ -1063,6 +1093,23 @@ FReply SSequencerPlaylistPanel::HandleAcceptDrop(const FDragDropEvent& DragDropE
 }
 
 
+void SSequencerPlaylistPanel::OnConcertSessionStartup(TSharedRef<IConcertClientSession> InSession)
+{
+	if (ensure(HeaderRow))
+	{
+		HeaderRow->SetShowGeneratedColumn(ColumnName_Preload, true);
+	}
+}
+
+void SSequencerPlaylistPanel::OnConcertSessionShutdown(TSharedRef<IConcertClientSession> InSession)
+{
+	if (ensure(HeaderRow))
+	{
+		HeaderRow->SetShowGeneratedColumn(ColumnName_Preload, false);
+	}
+}
+
+
 TSharedRef<FSequencerPlaylistItemDragDropOp> FSequencerPlaylistItemDragDropOp::New(const TArray<TSharedPtr<FSequencerPlaylistRowData>>& InSelectedItems)
 {
 	TSharedRef<FSequencerPlaylistItemDragDropOp> Operation = MakeShared<FSequencerPlaylistItemDragDropOp>();
@@ -1547,6 +1594,45 @@ TSharedRef<SWidget> SSequencerPlaylistItemWidget::GenerateWidgetForColumn(const
 			];
 	}
 
+	if (ColumnName == SSequencerPlaylistPanel::ColumnName_Preload)
+	{
+		return SNew(SVerticalBox)
+			+ SVerticalBox::Slot()
+			[
+				SNew(STextBlock)
+				.Font(FAppStyle::Get().GetFontStyle("FontAwesome.18"))
+				.Text(FEditorFontGlyphs::Circle)
+				.ColorAndOpacity_Lambda([this, WeakItem]() {
+					USequencerPlaylistsSubsystem* Subsystem = GEditor->GetEditorSubsystem<USequencerPlaylistsSubsystem>();
+					USequencerPlaylistItem_Sequence* SeqItem = Cast<USequencerPlaylistItem_Sequence>(WeakItem);
+					if (Subsystem && SeqItem)
+					{
+						const FTopLevelAssetPath SeqPath(SeqItem->GetSequence());
+						TPair<EConcertSequencerPreloadStatus, FText> Status = Subsystem->GetPreloadStatusForSequence(SeqPath);
+						switch (Status.Key)
+						{
+						case EConcertSequencerPreloadStatus::Succeeded: return FSlateColor(EStyleColor::Success);
+						case EConcertSequencerPreloadStatus::Pending: return FSlateColor(EStyleColor::Warning);
+						}
+					}
+
+					return FSlateColor(EStyleColor::Error);
+				})
+				.ToolTipText_Lambda([this, WeakItem]() {
+					USequencerPlaylistsSubsystem* Subsystem = GEditor->GetEditorSubsystem<USequencerPlaylistsSubsystem>();
+					USequencerPlaylistItem_Sequence* SeqItem = Cast<USequencerPlaylistItem_Sequence>(WeakItem);
+					if (Subsystem && SeqItem)
+					{
+						const FTopLevelAssetPath SeqPath(SeqItem->GetSequence());
+						TPair<EConcertSequencerPreloadStatus, FText> Status = Subsystem->GetPreloadStatusForSequence(SeqPath);
+						return Status.Value;
+					}
+
+					return FText::GetEmpty();
+				})
+			];
+	}
+
 	if (ColumnName == SSequencerPlaylistPanel::ColumnName_HoverDetails)
 	{
 		return SAssignNew(DetailsAnchor, SMenuAnchor)
diff --git a/Engine/Plugins/VirtualProduction/SequencerPlaylists/Source/SequencerPlaylists/Private/SequencerPlaylistsWidgets.h b/Engine/Plugins/VirtualProduction/SequencerPlaylists/Source/SequencerPlaylists/Private/SequencerPlaylistsWidgets.h
--- a/Engine/Plugins/VirtualProduction/SequencerPlaylists/Source/SequencerPlaylists/Private/SequencerPlaylistsWidgets.h
+++ b/Engine/Plugins/VirtualProduction/SequencerPlaylists/Source/SequencerPlaylists/Private/SequencerPlaylistsWidgets.h
@@ -12,7 +12,9 @@
 class STableViewBase;
 
 struct FAssetData;
+class IConcertClientSession;
 class SDockTab;
+class SHeaderRow;
 class SMenuAnchor;
 class SMultiLineEditableTextBox;
 class SSequencerPlaylistItemWidget;
@@ -51,6 +53,7 @@ public:
 	static const FName ColumnName_Offset;
 	static const FName ColumnName_Pause;
 	static const FName ColumnName_Loop;
+	static const FName ColumnName_Preload;
 	static const FName ColumnName_HoverDetails;
 
 	void Construct(const FArguments& InArgs, TSharedPtr<SDockTab> InContainingTab);
@@ -109,6 +112,10 @@ private:
 	TOptional<EItemDropZone> HandleCanAcceptDrop(const FDragDropEvent& DragDropEvent, EItemDropZone DropZone, TSharedPtr<FSequencerPlaylistRowData> RowData);
 	FReply HandleAcceptDrop(const FDragDropEvent& DragDropEvent, EItemDropZone DropZone, TSharedPtr<FSequencerPlaylistRowData> RowData);
 
+	// Multi-user event handling
+	void OnConcertSessionStartup(TSharedRef<IConcertClientSession> InSession);
+	void OnConcertSessionShutdown(TSharedRef<IConcertClientSession> InSession);
+
 private:
 	TWeakObjectPtr<USequencerPlaylistPlayer> WeakPlayer;
 	TWeakObjectPtr<USequencerPlaylist> WeakLoadedPlaylist;
@@ -118,6 +125,7 @@ private:
 	TWeakPtr<SDockTab> WeakContainingTab;
 	TSharedPtr<SSearchBox> SearchBox;
 	TSharedPtr<TTextFilter<const FSequencerPlaylistRowData&>> SearchTextFilter;
+	TSharedPtr<SHeaderRow> HeaderRow;
 	TArray<TSharedPtr<FSequencerPlaylistRowData>> ItemRows;
 	TSharedPtr<SListView<TSharedPtr<FSequencerPlaylistRowData>>> ItemListView;
 };
