[Concert] Adapting ConcertSyncServer / UnrealMultiUserServer to work with QuicMessaging

Adapting the ConcertSyncServer and the UnrealMultiUserServer configs and code to work with QuicMessaging as an alternative to UdpMessaging.

diff --git a/Engine/Plugins/Developer/Concert/ConcertSync/ConcertSyncServer/Source/ConcertSyncServer/Public/ConcertSyncServerLoop.inl b/Engine/Plugins/Developer/Concert/ConcertSync/ConcertSyncServer/Source/ConcertSyncServer/Public/ConcertSyncServerLoop.inl
--- a/Engine/Plugins/Developer/Concert/ConcertSync/ConcertSyncServer/Source/ConcertSyncServer/Public/ConcertSyncServerLoop.inl
+++ b/Engine/Plugins/Developer/Concert/ConcertSync/ConcertSyncServer/Source/ConcertSyncServer/Public/ConcertSyncServerLoop.inl
@@ -90,11 +90,15 @@ int32 ConcertSyncServerLoop(const TCHAR* CommandLine, const FConcertSyncServerLo
 		// Give external modules to do early initialisation
 		InitArgs.PreInitServerLoop.Broadcast();
 		
-		TSharedPtr<IPlugin> Plugin = IPluginManager::Get().FindPlugin(TEXT("UdpMessaging"));
-		if (!Plugin || !Plugin->IsEnabled())
+		TSharedPtr<IPlugin> UdpPlugin = IPluginManager::Get().FindPlugin(TEXT("UdpMessaging"));
+		TSharedPtr<IPlugin> QuicPlugin = IPluginManager::Get().FindPlugin(TEXT("QuicMessaging"));
+
+		// Either the UdpMessaging or QuicMessaging plugin
+		// should be added to the {appname}.Target.cs build file.
+		if (!UdpPlugin || !QuicPlugin || (!UdpPlugin->IsEnabled() && !QuicPlugin->IsEnabled()))
 		{
-			// The UdpMessaging plugin should be added to the {appname}.Target.cs build file.
-			UE_LOG(LogSyncServer, Warning, TEXT("The 'UDP Messaging' plugin is disabled. The Concert server only supports UDP protocol."));
+			UE_LOG(LogSyncServer, Warning, TEXT("The 'UDP Messaging' and 'QUIC Messaging' plugins are disabled."
+				"The Concert server only supports UDP/QUIC protocol."));
 		}
 
 		// Setup Concert Sync to run in server mode.
diff --git a/Engine/Programs/UnrealMultiUserServer/Config/DefaultEngine.ini b/Engine/Programs/UnrealMultiUserServer/Config/DefaultEngine.ini
--- a/Engine/Programs/UnrealMultiUserServer/Config/DefaultEngine.ini
+++ b/Engine/Programs/UnrealMultiUserServer/Config/DefaultEngine.ini
@@ -3,4 +3,5 @@ GameNameShort=MUS
 
 [Plugins]
 +ProgramEnabledPlugins="UdpMessaging"
++ProgramEnabledPlugins="QuicMessaging"
 +ProgramEnabledPlugins="ConcertSyncServer"
diff --git a/Engine/Programs/UnrealMultiUserSlateServer/Config/DefaultEngine.ini b/Engine/Programs/UnrealMultiUserSlateServer/Config/DefaultEngine.ini
--- a/Engine/Programs/UnrealMultiUserSlateServer/Config/DefaultEngine.ini
+++ b/Engine/Programs/UnrealMultiUserSlateServer/Config/DefaultEngine.ini
@@ -3,5 +3,6 @@ GameNameShort=MUSS
 
 [Plugins]
 +ProgramEnabledPlugins="UdpMessaging"
++ProgramEnabledPlugins="QuicMessaging"
 +ProgramEnabledPlugins="ConcertSyncServer"
 +ProgramEnabledPlugins="MultiUserServer"
diff --git a/Engine/Source/Programs/UnrealMultiUserServer/UnrealMultiUserServer.Target.cs b/Engine/Source/Programs/UnrealMultiUserServer/UnrealMultiUserServer.Target.cs
--- a/Engine/Source/Programs/UnrealMultiUserServer/UnrealMultiUserServer.Target.cs
+++ b/Engine/Source/Programs/UnrealMultiUserServer/UnrealMultiUserServer.Target.cs
@@ -11,6 +11,7 @@ public class UnrealMultiUserServerTarget : TargetRules
 		LinkType = TargetLinkType.Modular;
 		LaunchModuleName = "UnrealMultiUserServer";
 		AdditionalPlugins.Add("UdpMessaging");
+		AdditionalPlugins.Add("QuicMessaging");
 		AdditionalPlugins.Add("ConcertSyncServer");
 
 		// This app compiles against Core/CoreUObject, but not the Engine or Editor, so compile out Engine and Editor references from Core/CoreUObject
diff --git a/Engine/Source/Programs/UnrealMultiUserSlateServer/UnrealMultiUserSlateServer.Target.cs b/Engine/Source/Programs/UnrealMultiUserSlateServer/UnrealMultiUserSlateServer.Target.cs
--- a/Engine/Source/Programs/UnrealMultiUserSlateServer/UnrealMultiUserSlateServer.Target.cs
+++ b/Engine/Source/Programs/UnrealMultiUserSlateServer/UnrealMultiUserSlateServer.Target.cs
@@ -11,6 +11,7 @@ public class UnrealMultiUserSlateServerTarget : TargetRules
 		LinkType = TargetLinkType.Modular;
 		LaunchModuleName = "UnrealMultiUserSlateServer";
 		AdditionalPlugins.Add("UdpMessaging");
+		AdditionalPlugins.Add("QuicMessaging");
 		AdditionalPlugins.Add("ConcertSyncServer");
 		AdditionalPlugins.Add("MultiUserServer");
 
