[Concert] Modifying ConcertSyncServer transport plugin check

Modifying the ConcertSyncServer transport plugin check to shutdown if neither of the two supported plugins UdpMessaging/QuicMessaging is available.

diff --git a/Engine/Plugins/Developer/Concert/ConcertSync/ConcertSyncServer/Source/ConcertSyncServer/Public/ConcertSyncServerLoop.h b/Engine/Plugins/Developer/Concert/ConcertSync/ConcertSyncServer/Source/ConcertSyncServer/Public/ConcertSyncServerLoop.h
--- a/Engine/Plugins/Developer/Concert/ConcertSync/ConcertSyncServer/Source/ConcertSyncServer/Public/ConcertSyncServerLoop.h
+++ b/Engine/Plugins/Developer/Concert/ConcertSync/ConcertSyncServer/Source/ConcertSyncServer/Public/ConcertSyncServerLoop.h
@@ -5,6 +5,11 @@
 #include "CoreMinimal.h"
 #include "ConcertSyncServerLoopInitArgs.h"
 
+/**
+ * Utility method to shutdown the ConcertSyncServer.
+ */
+void ShutdownConcertSyncServer(const FString& ServiceFriendlyName);
+
 /**
  * Blocking main loop for running a Concert Sync server application.
  */
diff --git a/Engine/Plugins/Developer/Concert/ConcertSync/ConcertSyncServer/Source/ConcertSyncServer/Public/ConcertSyncServerLoop.inl b/Engine/Plugins/Developer/Concert/ConcertSync/ConcertSyncServer/Source/ConcertSyncServer/Public/ConcertSyncServerLoop.inl
--- a/Engine/Plugins/Developer/Concert/ConcertSync/ConcertSyncServer/Source/ConcertSyncServer/Public/ConcertSyncServerLoop.inl
+++ b/Engine/Plugins/Developer/Concert/ConcertSync/ConcertSyncServer/Source/ConcertSyncServer/Public/ConcertSyncServerLoop.inl
@@ -22,6 +22,25 @@
 
 DEFINE_LOG_CATEGORY_STATIC(LogSyncServer, Log, All);
 
+void ShutdownConcertSyncServer(const FString& ServiceFriendlyName)
+{
+	UE_LOG(LogSyncServer, Display, TEXT("%s Shutdown"), *ServiceFriendlyName);
+
+	// Allow the game thread to finish processing any latent tasks.
+	// They will be relying on what we are going to shutdown...
+	FTaskGraphInterface::Get().ProcessThreadUntilIdle(ENamedThreads::GameThread);
+
+	FEngineLoop::AppPreExit();
+
+	// Unloading Modules isn't handled by AppExit
+	FModuleManager::Get().UnloadModulesAtShutdown();
+
+	// On Mac AppExit is passed as delegate to the OS, which will call it implicitly
+#if !PLATFORM_MAC
+	FEngineLoop::AppExit();
+#endif
+}
+
 int32 ConcertSyncServerLoop(const TCHAR* CommandLine, const FConcertSyncServerLoopInitArgs& InitArgs)
 {
 	FTaskTagScope Scope(ETaskTag::EGameThread);
@@ -34,9 +53,18 @@ int32 ConcertSyncServerLoop(const TCHAR* CommandLine, const FConcertSyncServerLo
 
 	// start up the main loop, adding some extra command line arguments:
 	//	-Messaging enables MessageBus transports
-	int32 Result = GEngineLoop.PreInit(*FString::Printf(TEXT("%s %s"), CommandLine, TEXT(" -Messaging")));
+	const int32 Result = GEngineLoop.PreInit(*FString::Printf(TEXT("%s %s"), CommandLine, TEXT(" -Messaging")));
 	check(GConfig && GConfig->IsReadyForUse());
 
+	if (Result < 0)
+	{
+		UE_LOG(LogSyncServer, Error, TEXT("EngineLoop PreInit failed!"));
+
+		ShutdownConcertSyncServer(InitArgs.ServiceFriendlyName);
+
+		return Result;
+	}
+
 	if (InitArgs.bShowConsole)
 	{
 		GLogConsole->Show(true);
@@ -62,146 +90,141 @@ int32 ConcertSyncServerLoop(const TCHAR* CommandLine, const FConcertSyncServerLo
 	if (!IConcertSyncServerModule::IsAvailable())
 	{
 		UE_LOG(LogSyncServer, Error, TEXT("ConcertSyncServer Module is needed for proper execution. Initialization failed!"));
-		Result = -1;
+
+		ShutdownConcertSyncServer(InitArgs.ServiceFriendlyName);
+
+		return -1;
 	}
 
 	const UConcertServerConfig* ServerConfig = nullptr;
 
-	if (Result >= 0)
+	// Get the server config.
+	if (InitArgs.GetServerConfigFunc)
 	{
-		// Get the server config.
-		if (InitArgs.GetServerConfigFunc)
-		{
-			ServerConfig = InitArgs.GetServerConfigFunc();
-		}
-		else
-		{
-			ServerConfig = IConcertSyncServerModule::Get().ParseServerSettings(FCommandLine::Get());
-		}
-		if (!ServerConfig)
-		{
-			UE_LOG(LogSyncServer, Error, TEXT("%s Configuration Failed!"), *InitArgs.ServiceFriendlyName);
-			Result = -1;
-		}
+		ServerConfig = InitArgs.GetServerConfigFunc();
+	}
+	else
+	{
+		ServerConfig = IConcertSyncServerModule::Get().ParseServerSettings(FCommandLine::Get());
 	}
 
-	if (Result >= 0)
+	if (!ServerConfig)
 	{
-		// Give external modules to do early initialisation
-		InitArgs.PreInitServerLoop.Broadcast();
+		UE_LOG(LogSyncServer, Error, TEXT("%s Configuration Failed!"), *InitArgs.ServiceFriendlyName);
 		
-		TSharedPtr<IPlugin> UdpPlugin = IPluginManager::Get().FindPlugin(TEXT("UdpMessaging"));
-		TSharedPtr<IPlugin> QuicPlugin = IPluginManager::Get().FindPlugin(TEXT("QuicMessaging"));
+		ShutdownConcertSyncServer(InitArgs.ServiceFriendlyName);
 
-		// Either the UdpMessaging or QuicMessaging plugin
-		// should be added to the {appname}.Target.cs build file.
-		if (!UdpPlugin || !QuicPlugin || (!UdpPlugin->IsEnabled() && !QuicPlugin->IsEnabled()))
-		{
-			UE_LOG(LogSyncServer, Warning, TEXT("The 'UDP Messaging' and 'QUIC Messaging' plugins are disabled."
-				"The Concert server only supports UDP/QUIC protocol."));
-		}
+		return -1;
+	}
+
+
+	// Give external modules to do early initialisation
+	InitArgs.PreInitServerLoop.Broadcast();
+
+	TSharedPtr<IPlugin> UdpPlugin = IPluginManager::Get().FindPlugin(TEXT("UdpMessaging"));
+	TSharedPtr<IPlugin> QuicPlugin = IPluginManager::Get().FindPlugin(TEXT("QuicMessaging"));
+
+	// Either the UdpMessaging or QuicMessaging plugin
+	// should be added to the {appname}.Target.cs build file.
+	const bool bUdpEnabled = (UdpPlugin && UdpPlugin->IsEnabled());
+	const bool bQuicEnabled = (QuicPlugin && QuicPlugin->IsEnabled());
+
+	if (!bUdpEnabled && !bQuicEnabled)
+	{
+		UE_LOG(LogSyncServer, Error, TEXT("The 'UDP Messaging' and 'QUIC Messaging' plugins are disabled. "
+			"The Concert server only supports UDP/QUIC protocol."));
 
-		// Setup Concert Sync to run in server mode.
-		TSharedPtr<IConcertSyncServer> ConcertSyncServer = IConcertSyncServerModule::Get().CreateServer(InitArgs.ServiceRole, InitArgs.ServiceAutoArchiveSessionFilter);
-		ConcertSyncServer->SetFileSharingService(InitArgs.FileSharingService);
-		ConcertSyncServer->Startup(ServerConfig, InitArgs.SessionFlags);
+		ShutdownConcertSyncServer(InitArgs.ServiceFriendlyName);
 
-		// if we have a default session, set it up properly
-		if (!ServerConfig->DefaultSessionName.IsEmpty())
+		return -1;
+	}
+
+	// Setup Concert Sync to run in server mode.
+	TSharedPtr<IConcertSyncServer> ConcertSyncServer = IConcertSyncServerModule::Get().CreateServer(InitArgs.ServiceRole, InitArgs.ServiceAutoArchiveSessionFilter);
+	ConcertSyncServer->SetFileSharingService(InitArgs.FileSharingService);
+	ConcertSyncServer->Startup(ServerConfig, InitArgs.SessionFlags);
+
+	// if we have a default session, set it up properly
+	if (!ServerConfig->DefaultSessionName.IsEmpty())
+	{
+		IConcertServerRef ConcertServer = ConcertSyncServer->GetConcertServer();
+		if (!ConcertServer->GetLiveSessionIdByName(ServerConfig->DefaultSessionName).IsValid())
 		{
-			IConcertServerRef ConcertServer = ConcertSyncServer->GetConcertServer();
-			if (!ConcertServer->GetLiveSessionIdByName(ServerConfig->DefaultSessionName).IsValid())
-			{
-				FConcertSessionInfo SessionInfo = ConcertServer->CreateSessionInfo();
-				SessionInfo.SessionName = ServerConfig->DefaultSessionName;
-				SessionInfo.Settings = ServerConfig->DefaultSessionSettings;
-				SessionInfo.VersionInfos.Add(ServerConfig->DefaultVersionInfo);
+			FConcertSessionInfo SessionInfo = ConcertServer->CreateSessionInfo();
+			SessionInfo.SessionName = ServerConfig->DefaultSessionName;
+			SessionInfo.Settings = ServerConfig->DefaultSessionSettings;
+			SessionInfo.VersionInfos.Add(ServerConfig->DefaultVersionInfo);
 
-				bool bSuccess = false;
-				if (!ServerConfig->DefaultSessionToRestore.IsEmpty())
+			bool bSuccess = false;
+			if (!ServerConfig->DefaultSessionToRestore.IsEmpty())
+			{
+				const FGuid ArchivedSessionId = ConcertServer->GetArchivedSessionIdByName(ServerConfig->DefaultSessionToRestore);
+				if (ArchivedSessionId.IsValid())
 				{
-					const FGuid ArchivedSessionId = ConcertServer->GetArchivedSessionIdByName(ServerConfig->DefaultSessionToRestore);
-					if (ArchivedSessionId.IsValid())
-					{
-						FText FailureReason;
-						bSuccess = ConcertServer->RestoreSession(ArchivedSessionId, SessionInfo, FConcertSessionFilter(), FailureReason).IsValid();
-						if (!bSuccess)
-						{
-							UE_LOG(LogSyncServer, Error, TEXT("%s failed to restore archived session '%s' due to '%s'. Creating new session instead..."), *InitArgs.ServiceFriendlyName, *ServerConfig->DefaultSessionToRestore, *FailureReason.ToString());
-						}
-					}
-					else
+					FText FailureReason;
+					bSuccess = ConcertServer->RestoreSession(ArchivedSessionId, SessionInfo, FConcertSessionFilter(), FailureReason).IsValid();
+					if (!bSuccess)
 					{
-						UE_LOG(LogSyncServer, Error, TEXT("%s could not find archived session '%s' to restore. Creating new session instead..."), *InitArgs.ServiceFriendlyName, *ServerConfig->DefaultSessionToRestore);
+						UE_LOG(LogSyncServer, Error, TEXT("%s failed to restore archived session '%s' due to '%s'. Creating new session instead..."), *InitArgs.ServiceFriendlyName, *ServerConfig->DefaultSessionToRestore, *FailureReason.ToString());
 					}
 				}
+				else
+				{
+					UE_LOG(LogSyncServer, Error, TEXT("%s could not find archived session '%s' to restore. Creating new session instead..."), *InitArgs.ServiceFriendlyName, *ServerConfig->DefaultSessionToRestore);
+				}
+			}
 
+			if (!bSuccess)
+			{
+				FText FailureReason;
+				bSuccess = ConcertServer->CreateSession(SessionInfo, FailureReason).IsValid();
 				if (!bSuccess)
 				{
-					FText FailureReason;
-					bSuccess = ConcertServer->CreateSession(SessionInfo, FailureReason).IsValid();
-					if (!bSuccess)
-					{
-						UE_LOG(LogSyncServer, Error, TEXT("%s failed to create session '%s' due to '%s'."), *InitArgs.ServiceFriendlyName, *ServerConfig->DefaultSessionName, *FailureReason.ToString());
-					}
+					UE_LOG(LogSyncServer, Error, TEXT("%s failed to create session '%s' due to '%s'."), *InitArgs.ServiceFriendlyName, *ServerConfig->DefaultSessionName, *FailureReason.ToString());
 				}
 			}
 		}
-		
-		UE_LOG(LogSyncServer, Display, TEXT("%s Initialized (Name: %s, Version: %d.%d, Role: %s)"), *InitArgs.ServiceFriendlyName, *ConcertSyncServer->GetConcertServer()->GetServerInfo().ServerName, ENGINE_MAJOR_VERSION, ENGINE_MINOR_VERSION, *ConcertSyncServer->GetConcertServer()->GetRole());
-		InitArgs.PostInitServerLoop.Broadcast(ConcertSyncServer.ToSharedRef());
-		
-		double LastTime = FPlatformTime::Seconds();
-		const float IdealFrameTime = 1.0f / InitArgs.IdealFramerate;
+	}
 
-		while (!IsEngineExitRequested())
-		{
-			const double CurrentTime = FPlatformTime::Seconds();
-			const double DeltaTime = CurrentTime - LastTime;
+	UE_LOG(LogSyncServer, Display, TEXT("%s Initialized (Name: %s, Version: %d.%d, Role: %s)"), *InitArgs.ServiceFriendlyName, *ConcertSyncServer->GetConcertServer()->GetServerInfo().ServerName, ENGINE_MAJOR_VERSION, ENGINE_MINOR_VERSION, *ConcertSyncServer->GetConcertServer()->GetRole());
+	InitArgs.PostInitServerLoop.Broadcast(ConcertSyncServer.ToSharedRef());
+
+	double LastTime = FPlatformTime::Seconds();
+	const float IdealFrameTime = 1.0f / InitArgs.IdealFramerate;
+
+	while (!IsEngineExitRequested())
+	{
+		const double CurrentTime = FPlatformTime::Seconds();
+		const double DeltaTime = CurrentTime - LastTime;
 
-			FTaskGraphInterface::Get().ProcessThreadUntilIdle(ENamedThreads::GameThread);
+		FTaskGraphInterface::Get().ProcessThreadUntilIdle(ENamedThreads::GameThread);
 
-			// Pump & Tick objects
-			InitArgs.TickPostGameThread.Broadcast(DeltaTime);
-			FTSTicker::GetCoreTicker().Tick(DeltaTime);
+		// Pump & Tick objects
+		InitArgs.TickPostGameThread.Broadcast(DeltaTime);
+		FTSTicker::GetCoreTicker().Tick(DeltaTime);
 
-			GFrameCounter++;
-			FStats::AdvanceFrame(false);
-			GLog->FlushThreadedLogs();
+		GFrameCounter++;
+		FStats::AdvanceFrame(false);
+		GLog->FlushThreadedLogs();
 
-			// Run garbage collection for the UObjects for the rest of the frame or at least to 2 ms
-			IncrementalPurgeGarbage(true, FMath::Max<float>(0.002f, IdealFrameTime - (FPlatformTime::Seconds() - LastTime)));
+		// Run garbage collection for the UObjects for the rest of the frame or at least to 2 ms
+		IncrementalPurgeGarbage(true, FMath::Max<float>(0.002f, IdealFrameTime - (FPlatformTime::Seconds() - LastTime)));
 
 #if PLATFORM_MAC
-			// Pumps the message from the main loop. (On Mac, we have a full application to get a proper console window to output logs)
-			FPlatformApplicationMisc::PumpMessages(true);
+		// Pumps the message from the main loop. (On Mac, we have a full application to get a proper console window to output logs)
+		FPlatformApplicationMisc::PumpMessages(true);
 #endif
 
-			// Throttle main thread main fps by sleeping if we still have time
-			FPlatformProcess::Sleep(FMath::Max<float>(0.0f, IdealFrameTime - (FPlatformTime::Seconds() - LastTime)));
+		// Throttle main thread main fps by sleeping if we still have time
+		FPlatformProcess::Sleep(FMath::Max<float>(0.0f, IdealFrameTime - (FPlatformTime::Seconds() - LastTime)));
 
-			LastTime = CurrentTime;
-		}
-
-		ConcertSyncServer->Shutdown();
-		ConcertSyncServer.Reset();
+		LastTime = CurrentTime;
 	}
 
-	UE_LOG(LogSyncServer, Display, TEXT("%s Shutdown"), *InitArgs.ServiceFriendlyName);
-
-	// Allow the game thread to finish processing any latent tasks.
-	// They will be relying on what we are going to shutdown...
-	FTaskGraphInterface::Get().ProcessThreadUntilIdle(ENamedThreads::GameThread);
-
-	FEngineLoop::AppPreExit();
-
-	// Unloading Modules isn't handled by AppExit
-	FModuleManager::Get().UnloadModulesAtShutdown();
+	ConcertSyncServer->Shutdown();
+	ConcertSyncServer.Reset();
 
-	// On Mac AppExit is passed as delegate to the OS, which will call it implicitly
-#if !PLATFORM_MAC
-	FEngineLoop::AppExit();
-#endif
+	ShutdownConcertSyncServer(InitArgs.ServiceFriendlyName);
 
 	return Result;
 }
