Add console command for logging replication streams and authority on MU server

diff --git a/Engine/Plugins/Developer/Concert/ConcertSync/ConcertSyncServer/Source/ConcertSyncServer/Private/ConcertSyncServer.cpp b/Engine/Plugins/Developer/Concert/ConcertSync/ConcertSyncServer/Source/ConcertSyncServer/Private/ConcertSyncServer.cpp
--- a/Engine/Plugins/Developer/Concert/ConcertSync/ConcertSyncServer/Source/ConcertSyncServer/Private/ConcertSyncServer.cpp
+++ b/Engine/Plugins/Developer/Concert/ConcertSync/ConcertSyncServer/Source/ConcertSyncServer/Private/ConcertSyncServer.cpp
@@ -229,8 +229,15 @@ bool MigrateSessionData(const FConcertSyncSessionDatabase& InSourceDatabase, con
 FConcertSyncServer::FConcertSyncServer(const FString& InRole, const FConcertSessionFilter& InAutoArchiveSessionFilter)
 	: ConcertServer(IConcertServerModule::Get().CreateServer(InRole, InAutoArchiveSessionFilter, this))
 	, SessionFlags(EConcertSyncSessionFlags::None)
-{
-}
+	, LogReplicationStreamsConsoleCommand(
+		TEXT("Replication.LogStreams"),
+		TEXT("Logs the streams registered by all clients"),
+		FConsoleCommandDelegate::CreateRaw(this, &FConcertSyncServer::LogReplicationStreams))
+	, LogReplicationAuthorityConsoleCommand(
+		TEXT("Replication.LogAutority"),
+		TEXT("Logs the authority of the clients"),
+		FConsoleCommandDelegate::CreateRaw(this, &FConcertSyncServer::LogReplicationAuthority))
+{}
 
 FConcertSyncServer::~FConcertSyncServer()
 {
@@ -591,4 +598,88 @@ void FConcertSyncServer::DestroyArchivedSession(const FGuid& InArchivedSessionId
 	ArchivedSessions.Remove(InArchivedSessionId);
 }
 
+void FConcertSyncServer::LogReplicationStreams() const
+{
+	UE_LOG(LogConcert, Log, TEXT("Printing replication streams..."));
+	
+	TWideStringBuilder<2048> StringBuilder;
+	constexpr auto SessionIntent = TEXT("\n");
+	constexpr auto ClientIndent = TEXT("\n\t");
+	constexpr auto StreamIndent = TEXT("\n\t\t");
+	constexpr auto ObjectIndent = TEXT("\n\t\t\t");
+	constexpr auto PropertyIndent = TEXT("\n\t\t\t\t");
+	
+	for (const TPair<FGuid, TSharedRef<UE::ConcertSyncServer::Replication::FConcertServerReplicationManager>>& ReplicationManagerPair : LiveSessionReplicationManagers)
+	{
+		const TSharedPtr<FConcertSyncServerLiveSession>* LiveSession = LiveSessions.Find(ReplicationManagerPair.Key);
+		check(LiveSession);
+
+		StringBuilder << SessionIntent << TEXT("Session: ") << *LiveSession->Get()->GetSession().GetName() << TEXT(" (") << *ReplicationManagerPair.Key.ToString() << TEXT(")");
+		for (const FConcertSessionClientInfo& ClientInfo : LiveSession->Get()->GetSession().GetSessionClients())
+		{
+			const FGuid ClientId = ClientInfo.ClientEndpointId;
+			StringBuilder << ClientIndent << TEXT("Client: ") << *ClientInfo.ClientInfo.DisplayName << TEXT(" (") << ClientId.ToString() << TEXT(")");
+			ReplicationManagerPair.Value->ForEachStream(ClientId, [&StringBuilder, StreamIndent, ObjectIndent, PropertyIndent](const FReplicationStreamDescription& Stream)
+			{
+				StringBuilder << StreamIndent << TEXT("Stream: ") << *Stream.BaseDescription.Identifier.ToString();
+				for (const TPair<FSoftObjectPath, FReplicatedObjectInfo>& ReplicatedObject : Stream.BaseDescription.ReplicationMap.ReplicatedObjects)
+				{
+					StringBuilder << ObjectIndent << *ReplicatedObject.Key.ToString();
+					for (const FConcertPropertyChain& Property : ReplicatedObject.Value.PropertySelection.ReplicatedProperties)
+					{
+						StringBuilder << PropertyIndent << Property.ToString();
+					}
+				}
+				return EBreakBehavior::Continue;
+			});
+		}
+	}
+	
+	UE_LOG(LogConcert, Log, TEXT("%s"), StringBuilder.ToString());
+	// The logs before and after are to have some feedback in case nothing is printed (e.g. no live sessions).
+	UE_LOG(LogConcert, Log, TEXT("Done."));
+}
+
+void FConcertSyncServer::LogReplicationAuthority() const
+{
+	using namespace UE::ConcertSyncServer::Replication;
+	UE_LOG(LogConcert, Log, TEXT("Printing replication authority..."));
+	
+	TWideStringBuilder<2048> StringBuilder;
+	constexpr auto SessionIntent = TEXT("\n");
+	constexpr auto ClientIndent = TEXT("\n\t");
+	constexpr auto StreamIndent = TEXT("\n\t\t");
+	constexpr auto ObjectIndent = TEXT("\n\t\t\t");
+	
+	for (const TPair<FGuid, TSharedRef<FConcertServerReplicationManager>>& ReplicationManagerPair : LiveSessionReplicationManagers)
+	{
+		const TSharedPtr<FConcertSyncServerLiveSession>* LiveSession = LiveSessions.Find(ReplicationManagerPair.Key);
+		check(LiveSession);
+
+		StringBuilder << SessionIntent << TEXT("Session: ") << *LiveSession->Get()->GetSession().GetName() << TEXT(" (") << *ReplicationManagerPair.Key.ToString() << TEXT(")");
+		for (const FConcertSessionClientInfo& ClientInfo : LiveSession->Get()->GetSession().GetSessionClients())
+		{
+			const FGuid ClientId = ClientInfo.ClientEndpointId;
+			StringBuilder << ClientIndent << TEXT("Client: ") << *ClientInfo.ClientInfo.DisplayName << TEXT(" (") << ClientId.ToString() << TEXT(")");
+			
+			const TSharedRef<FConcertServerReplicationManager>& ReplicationManager = ReplicationManagerPair.Value;
+			ReplicationManager->ForEachStream(ClientId, [&StringBuilder, StreamIndent, ObjectIndent, ClientId, &ReplicationManager](const FReplicationStreamDescription& Stream)
+			{
+				const FGuid StreamId = Stream.BaseDescription.Identifier;
+				StringBuilder << StreamIndent << TEXT("Stream: ") << *StreamId.ToString();
+				ReplicationManager->GetAuthorityManager().EnumerateAuthority(ClientId, StreamId, [&StringBuilder, ObjectIndent](const FSoftObjectPath& AuthoredObject)
+				{
+					StringBuilder << ObjectIndent << *AuthoredObject.ToString();
+					return EBreakBehavior::Continue;
+				});
+				return EBreakBehavior::Continue;
+			});
+		}
+	}
+	
+	UE_LOG(LogConcert, Log, TEXT("%s"), StringBuilder.ToString());
+	// The logs before and after are to have some feedback in case nothing is printed (e.g. no live sessions).
+	UE_LOG(LogConcert, Log, TEXT("Done."));
+}
+
 #undef LOCTEXT_NAMESPACE
diff --git a/Engine/Plugins/Developer/Concert/ConcertSync/ConcertSyncServer/Source/ConcertSyncServer/Private/ConcertSyncServer.h b/Engine/Plugins/Developer/Concert/ConcertSync/ConcertSyncServer/Source/ConcertSyncServer/Private/ConcertSyncServer.h
--- a/Engine/Plugins/Developer/Concert/ConcertSync/ConcertSyncServer/Source/ConcertSyncServer/Private/ConcertSyncServer.h
+++ b/Engine/Plugins/Developer/Concert/ConcertSync/ConcertSyncServer/Source/ConcertSyncServer/Private/ConcertSyncServer.h
@@ -5,6 +5,7 @@
 #include "ConcertServerEventForwardingSink.h"
 #include "IConcertSyncServer.h"
 #include "IConcertServerEventSink.h"
+#include "HAL/IConsoleManager.h"
 
 class IConcertServerSession;
 class FConcertServerWorkspace;
@@ -26,6 +27,7 @@ struct FConcertSessionFilter;
 class FConcertSyncServer : public IConcertSyncServer, public TConcertServerEventForwardingSink<FConcertSyncServer>
 {
 public:
+	
 	FConcertSyncServer(const FString& InRole, const FConcertSessionFilter& InAutoArchiveSessionFilter);
 	virtual ~FConcertSyncServer();
 
@@ -93,4 +95,14 @@ private:
 
 	/** Optional side channel to exchange large blobs (package data) with the server in a scalable way (ex. the request/response transport layer is not designed and doesn't support exchanging 3GB packages). */
 	TSharedPtr<IConcertFileSharingService> FileSharingService;
+
+	/** Logs the replication streams registered by all clients */
+	FAutoConsoleCommand LogReplicationStreamsConsoleCommand;
+	/** Logs the replication authority of the clients */
+	FAutoConsoleCommand LogReplicationAuthorityConsoleCommand;
+
+	/** Logs the replication streams registered by all clients */
+	void LogReplicationStreams() const;
+	/** Logs the replication authority of the clients */
+	void LogReplicationAuthority() const;
 };
diff --git a/Engine/Plugins/Developer/Concert/ConcertSync/ConcertSyncServer/Source/ConcertSyncServer/Private/Replication/AuthorityManager.cpp b/Engine/Plugins/Developer/Concert/ConcertSync/ConcertSyncServer/Source/ConcertSyncServer/Private/Replication/AuthorityManager.cpp
--- a/Engine/Plugins/Developer/Concert/ConcertSync/ConcertSyncServer/Source/ConcertSyncServer/Private/Replication/AuthorityManager.cpp
+++ b/Engine/Plugins/Developer/Concert/ConcertSync/ConcertSyncServer/Source/ConcertSyncServer/Private/Replication/AuthorityManager.cpp
@@ -46,6 +46,24 @@ namespace UE::ConcertSyncServer::Replication
 		return OwnedObjects && OwnedObjects->Contains(ObjectChange.Object);
 	}
 
+	void FAuthorityManager::EnumerateAuthority(const FClientId& ClientId, const FStreamId& StreamId, TFunctionRef<EBreakBehavior(const FSoftObjectPath& Object)> Callback) const
+	{
+		const FClientAuthorityData* AuthorityData = ClientAuthorityData.Find(ClientId);
+		const TSet<FSoftObjectPath>* OwnedObjects = AuthorityData ? AuthorityData->OwnedObjects.Find(StreamId) : nullptr;
+		if (!OwnedObjects)
+		{
+			return;
+		}
+		
+		for (const FSoftObjectPath& AuthoredObject : *OwnedObjects)
+		{
+			if (Callback(AuthoredObject) == EBreakBehavior::Break)
+			{
+				break;
+			}
+		}
+	}
+
 	FAuthorityManager::EAuthorityResult FAuthorityManager::EnumerateAuthorityConflicts(
 		const FReplicatedObjectId& Object, 
 		const FConcertPropertySelection* OverwriteProperties,
diff --git a/Engine/Plugins/Developer/Concert/ConcertSync/ConcertSyncServer/Source/ConcertSyncServer/Private/Replication/AuthorityManager.h b/Engine/Plugins/Developer/Concert/ConcertSync/ConcertSyncServer/Source/ConcertSyncServer/Private/Replication/AuthorityManager.h
--- a/Engine/Plugins/Developer/Concert/ConcertSync/ConcertSyncServer/Source/ConcertSyncServer/Private/Replication/AuthorityManager.h
+++ b/Engine/Plugins/Developer/Concert/ConcertSync/ConcertSyncServer/Source/ConcertSyncServer/Private/Replication/AuthorityManager.h
@@ -54,6 +54,9 @@ namespace UE::ConcertSyncServer::Replication
 		 */
 		bool HasAuthorityToChange(const FReplicatedObjectId& ObjectChange) const;
 
+		/** Utility for iterating authority a client has for a given stream. */
+		void EnumerateAuthority(const FClientId& ClientId, const FStreamId& StreamId, TFunctionRef<EBreakBehavior(const FSoftObjectPath& Object)> Callback) const;
+
 		enum class EAuthorityResult : uint8
 		{
 			/** The client is allowed to take authority */
diff --git a/Engine/Plugins/Developer/Concert/ConcertSync/ConcertSyncServer/Source/ConcertSyncServer/Private/Replication/ConcertServerReplicationManager.h b/Engine/Plugins/Developer/Concert/ConcertSync/ConcertSyncServer/Source/ConcertSyncServer/Private/Replication/ConcertServerReplicationManager.h
--- a/Engine/Plugins/Developer/Concert/ConcertSync/ConcertSyncServer/Source/ConcertSyncServer/Private/Replication/ConcertServerReplicationManager.h
+++ b/Engine/Plugins/Developer/Concert/ConcertSync/ConcertSyncServer/Source/ConcertSyncServer/Private/Replication/ConcertServerReplicationManager.h
@@ -42,6 +42,8 @@ namespace UE::ConcertSyncServer::Replication
 		explicit FConcertServerReplicationManager(TSharedRef<IConcertServerSession> InLiveSession);
 		virtual ~FConcertServerReplicationManager() override;
 
+		const FAuthorityManager& GetAuthorityManager() const { return AuthorityManager.Get(); }
+
 		//~ Begin IAuthorityManagerGetters Interface
 		virtual void ForEachStream(const FGuid& ClientEndpointId, TFunctionRef<EBreakBehavior(const FReplicationStreamDescription& Stream)> Callback) override;
 		virtual void ForEachSendingClient(TFunctionRef<EBreakBehavior(const FGuid& ClientEndpointId)> Callback) override;
