Fix replication system not being shutdown when session is destroyed

diff --git a/Engine/Plugins/Developer/Concert/ConcertSync/ConcertSyncServer/Source/ConcertSyncServer/Private/ConcertSyncServer.cpp b/Engine/Plugins/Developer/Concert/ConcertSync/ConcertSyncServer/Source/ConcertSyncServer/Private/ConcertSyncServer.cpp
--- a/Engine/Plugins/Developer/Concert/ConcertSync/ConcertSyncServer/Source/ConcertSyncServer/Private/ConcertSyncServer.cpp
+++ b/Engine/Plugins/Developer/Concert/ConcertSync/ConcertSyncServer/Source/ConcertSyncServer/Private/ConcertSyncServer.cpp
@@ -601,6 +601,7 @@ void FConcertSyncServer::DestroyLiveSession(const TSharedRef<IConcertServerSessi
 	{
 		DestroyWorkspace(LiveSession.ToSharedRef());
 		DestroySequencerManager(LiveSession.ToSharedRef());
+		DestroyReplicationManager(InSession->GetId());
 		LiveSessions.Remove(InSession->GetId());
 	}
 }
