Replace references to EConcertReplicationMuteOptionIncludeSubobjects to EConcertReplicationMuteOption::ObjectAndSubobjects

diff --git a/Engine/Plugins/Developer/Concert/ConcertSync/ConcertSyncCore/Source/ConcertSyncCore/Public/Replication/Messages/Muting.h b/Engine/Plugins/Developer/Concert/ConcertSync/ConcertSyncCore/Source/ConcertSyncCore/Public/Replication/Messages/Muting.h
--- a/Engine/Plugins/Developer/Concert/ConcertSync/ConcertSyncCore/Source/ConcertSyncCore/Public/Replication/Messages/Muting.h
+++ b/Engine/Plugins/Developer/Concert/ConcertSync/ConcertSyncCore/Source/ConcertSyncCore/Public/Replication/Messages/Muting.h
@@ -73,12 +73,12 @@ struct FConcertReplication_ChangeMuteState_Request
 	/**
 	 * The objects to explicitly unmute.
 	 * 
-	 * If a specified object is implicitly muted, i.e. one of its outers is muted with the EConcertReplicationMuteFlags::IncludeSubobjects flag,
-	 * this will allow it to replicate again. If you specify the EConcertReplicationMuteFlags::IncludeSubobjects flag, then its subobjects will
+	 * If a specified object is implicitly muted, i.e. one of its outers is muted with the EConcertReplicationMuteFlags::ObjectAndSubobjects flag,
+	 * this will allow it to replicate again. If you specify the EConcertReplicationMuteFlags::ObjectAndSubobjects flag, then its subobjects will
 	 * also be allowed to be replicated again.
 	 *
 	 * If a specified object is explicitly muted, then it and all the subobjects that were implicitly muted because of it are unmuted;
-	 * in this case, it does not matter whether EConcertReplicationMuteFlags::IncludeSubobjects flag is set.
+	 * in this case, it does not matter whether EConcertReplicationMuteFlags::ObjectAndSubobjects flag is set.
 	 */
 	UPROPERTY()
 	TMap<FSoftObjectPath, FConcertReplication_ObjectMuteSetting> ObjectsToUnmute;
@@ -158,19 +158,19 @@ struct FConcertReplication_QueryMuteState_Response
 	TMap<FSoftObjectPath, FConcertReplication_ObjectMuteSetting> ExplicitlyMutedObjects;
 
 	/**
-	 * After an outer object is muted with EConcertReplicationMuteFlags::IncludeSubobjects, this contains the objects that were explicitly unmuted,
+	 * After an outer object is muted with EConcertReplicationMuteFlags::ObjectAndSubobjects, this contains the objects that were explicitly unmuted,
 	 * i.e. specified with FConcertReplication_ChangeMuteState_Request::ObjectsToUnmute.
 	 */
 	UPROPERTY()
 	TMap<FSoftObjectPath, FConcertReplication_ObjectMuteSetting> ExplicitlyUnmutedObjects;
 
 	/**
-	 * Objects that were implicitly muted, i.e. some their outers are explicitly muted with EConcertReplicationMuteFlags::IncludeSubobjects.
+	 * Objects that were implicitly muted, i.e. some their outers are explicitly muted with EConcertReplicationMuteFlags::ObjectAndSubobjects.
 	 */
 	UPROPERTY()
 	TSet<FSoftObjectPath> ImplicitlyMutedObjects;
 
-	/** Objects that were implicitly unmuted, i.e. 1. a parent is muted and 2. a child of parent is unmuted with EConcertReplicationMuteFlags::IncludeSubobjects */
+	/** Objects that were implicitly unmuted, i.e. 1. a parent is muted and 2. a child of parent is unmuted with EConcertReplicationMuteFlags::ObjectAndSubobjects */
 	UPROPERTY()
 	TSet<FSoftObjectPath> ImplicitlyUnmutedObjects;
 
diff --git a/Engine/Plugins/Developer/Concert/ConcertSync/ConcertSyncServer/Source/ConcertSyncServer/Private/Replication/MuteManager.cpp b/Engine/Plugins/Developer/Concert/ConcertSync/ConcertSyncServer/Source/ConcertSyncServer/Private/Replication/MuteManager.cpp
--- a/Engine/Plugins/Developer/Concert/ConcertSync/ConcertSyncServer/Source/ConcertSyncServer/Private/Replication/MuteManager.cpp
+++ b/Engine/Plugins/Developer/Concert/ConcertSync/ConcertSyncServer/Source/ConcertSyncServer/Private/Replication/MuteManager.cpp
@@ -318,7 +318,7 @@ namespace UE::ConcertSyncServer::Replication
 			const auto[AffectingParent, ParentMuteState] = FindAffectingParentState(ObjectToUnmute);
 			// ParentMuteState can either be:
 			//	- EMuteState::ExplicitlyMuted: all children should be muted
-			//	- EMuteState::ExplicitlyUnmuted: by invariant, there must be a parent that is ExplicitlyMuted with EConcertReplicationMuteFlags::IncludeSubobjects.
+			//	- EMuteState::ExplicitlyUnmuted: by invariant, there must be a parent that is ExplicitlyMuted with EConcertReplicationMuteFlags::ObjectAndSubobjects.
 			// Hence in both cases, children need to be updated to be marked as implicitly muted.
 			if (ParentMuteState)
 			{
diff --git a/Engine/Plugins/Developer/Concert/ConcertSync/ConcertSyncServer/Source/ConcertSyncServer/Private/Replication/MuteManager.h b/Engine/Plugins/Developer/Concert/ConcertSync/ConcertSyncServer/Source/ConcertSyncServer/Private/Replication/MuteManager.h
--- a/Engine/Plugins/Developer/Concert/ConcertSync/ConcertSyncServer/Source/ConcertSyncServer/Private/Replication/MuteManager.h
+++ b/Engine/Plugins/Developer/Concert/ConcertSync/ConcertSyncServer/Source/ConcertSyncServer/Private/Replication/MuteManager.h
@@ -86,12 +86,12 @@ namespace UE::ConcertSyncServer::Replication
 		{
 			/** Object is muted */
 			ExplicitlyMuted,
-			/** A parent object is ExplicitlyMuted and has EConcertReplicationMuteFlags::IncludeSubobjects set but is object and its children are supposed to be implicitly unmuted.  */
+			/** A parent object is ExplicitlyMuted and has EConcertReplicationMuteFlags::ObjectAndSubobjects set but is object and its children are supposed to be implicitly unmuted.  */
 			ExplicitlyUnmuted,
 			
-			/** Object is muted because one of its parent objects is ExplicitlyMuted and EConcertReplicationMuteFlags::IncludeSubobjects set. */
+			/** Object is muted because one of its parent objects is ExplicitlyMuted and EConcertReplicationMuteFlags::ObjectAndSubobjects set. */
 			ImplicitlyMuted,
-			/** Object is unmuted because one of its parent objects is ExplicitlyUnmuted and has EConcertReplicationMuteFlags::IncludeSubobjects set. */
+			/** Object is unmuted because one of its parent objects is ExplicitlyUnmuted and has EConcertReplicationMuteFlags::ObjectAndSubobjects set. */
 			ImplicitlyUnmuted,
 		};
 		
@@ -116,7 +116,7 @@ namespace UE::ConcertSyncServer::Replication
 		 */
 		TMap<FSoftObjectPath, FMuteData> MuteState;
 		
-		/** For each added object, checks whether it is a subobject of an object that is muted with the EConcertReplicationMuteFlags::IncludeSubobjects setting. */
+		/** For each added object, checks whether it is a subobject of an object that is muted with the EConcertReplicationMuteFlags::ObjectAndSubobjects setting. */
 		void TrackAddedSubobjectsForImplicitMuting(TConstArrayView<FConcertObjectInStreamID> AddedObjects);
 		/** Calls RemoveMuteState on every object that is unreferenced by client streams (directly or indirectly via a subobject). */
 		void UnmuteObjectsIfUnreferenced(TConstArrayView<FConcertObjectInStreamID> RemovedObjects);
@@ -139,10 +139,10 @@ namespace UE::ConcertSyncServer::Replication
 		void RemoveMuteState(const FSoftObjectPath& Object, TCallback&& OnRemoved);
 		void RemoveMuteState(const FSoftObjectPath& Object) { RemoveMuteState(Object, [](const FSoftObjectPath&){}); } 
 		
-		/** Sets the mute state of all subobjects of Parent until reaching a subobject that also has EConcertReplicationMuteFlags::IncludeSubobjects set. */
+		/** Sets the mute state of all subobjects of Parent until reaching a subobject that also has EConcertReplicationMuteFlags::ObjectAndSubobjects set. */
 		void UpdateImplicitStateUnder(const FSoftObjectPath& Parent, EMuteState NewImplicitState);
 		
-		/** Removes all subobject mute state affected by Parent's mute state assuming that ParentObject has no parents with the EConcertReplicationMuteFlags::IncludeSubobjects. */
+		/** Removes all subobject mute state affected by Parent's mute state assuming that ParentObject has no parents with the EConcertReplicationMuteFlags::ObjectAndSubobjects. */
 		template<CObjectProcessable TCallback>
 		void ClearAllChildStateUnder(const FSoftObjectPath& ParentObject, TCallback&& OnRemoved);
 		void ClearAllChildStateUnder(const FSoftObjectPath& ParentObject) { ClearAllChildStateUnder(ParentObject, [](const FSoftObjectPath&){}); }
diff --git a/Engine/Plugins/Developer/Concert/ConcertSync/ConcertSyncTest/Source/ConcertSyncTest/Private/Replication/Muting/MutingRejection.spec.cpp b/Engine/Plugins/Developer/Concert/ConcertSync/ConcertSyncTest/Source/ConcertSyncTest/Private/Replication/Muting/MutingRejection.spec.cpp
--- a/Engine/Plugins/Developer/Concert/ConcertSync/ConcertSyncTest/Source/ConcertSyncTest/Private/Replication/Muting/MutingRejection.spec.cpp
+++ b/Engine/Plugins/Developer/Concert/ConcertSync/ConcertSyncTest/Source/ConcertSyncTest/Private/Replication/Muting/MutingRejection.spec.cpp
@@ -69,7 +69,7 @@ namespace UE::ConcertSyncTests::Replication
 					});
 			});
 
-			It("Reject: Explicitly muting unreferenced object (EConcertReplicationMuteFlags::IncludeSubobjects) without subobject", [this]()
+			It("Reject: Explicitly muting unreferenced object (EConcertReplicationMuteFlags::ObjectAndSubobjects) without subobject", [this]()
 			{
 				Sender->JoinReplication();
 				Sender->GetClientReplicationManager()
@@ -80,7 +80,7 @@ namespace UE::ConcertSyncTests::Replication
 					});
 			});
 
-			It("Reject: Explicitly unmuting unreferenced object (EConcertReplicationMuteFlags::IncludeSubobjects) without subobject", [this]()
+			It("Reject: Explicitly unmuting unreferenced object (EConcertReplicationMuteFlags::ObjectAndSubobjects) without subobject", [this]()
 			{
 				Sender->JoinReplication();
 				Sender->GetClientReplicationManager()
@@ -91,7 +91,7 @@ namespace UE::ConcertSyncTests::Replication
 					});
 			});
 			
-			It("Allow: Explicitly muting unreferenced object (EConcertReplicationMuteFlags::IncludeSubobjects) with subobject", [this]()
+			It("Allow: Explicitly muting unreferenced object (EConcertReplicationMuteFlags::ObjectAndSubobjects) with subobject", [this]()
 			{
 				Sender->JoinReplication(Replicator_FooSubobject->CreateSenderArgs());
 				Sender->GetClientReplicationManager()
@@ -116,13 +116,13 @@ namespace UE::ConcertSyncTests::Replication
 					});
 			});
 			
-			It("Allow: Explicitly unmuting unreferenced object (EConcertReplicationMuteFlags::IncludeSubobjects) with subobject", [this]()
+			It("Allow: Explicitly unmuting unreferenced object (EConcertReplicationMuteFlags::ObjectAndSubobjects) with subobject", [this]()
 			{
 				Sender->JoinReplication(Replicator_FooSubobject->CreateSenderArgs());
 
 				// Check nothing is muted if ...
 				Sender->GetClientReplicationManager().MuteObjects({ ObjectReplicator_Foo->TestObject}, EConcertReplicationMuteOption::ObjectAndSubobjects);
-				// ... we unmute with EConcertReplicationMuteFlags::IncludeSubobjects flag
+				// ... we unmute with EConcertReplicationMuteFlags::ObjectAndSubobjects flag
 				Sender->GetClientReplicationManager().UnmuteObjects({ ObjectReplicator_Foo->TestObject}, EConcertReplicationMuteOption::ObjectAndSubobjects);
 				Sender->GetClientReplicationManager()
 					.QueryMuteState()
