Extend mute manager API:
- ValidateRequest: Allows validating a request with an optional future server state of registered objects
- ApplyRequestAndEnumerateSyncControl: Applies a mute request but skips sending network messages to affected clients (so the caller can send them manually)
This is in preparation for implementing FConcertReplication_PutState_Request.

diff --git a/Engine/Plugins/Developer/Concert/ConcertSync/ConcertSyncServer/Source/ConcertSyncServer/Private/Replication/MuteManager.cpp b/Engine/Plugins/Developer/Concert/ConcertSync/ConcertSyncServer/Source/ConcertSyncServer/Private/Replication/MuteManager.cpp
--- a/Engine/Plugins/Developer/Concert/ConcertSync/ConcertSyncServer/Source/ConcertSyncServer/Private/Replication/MuteManager.cpp
+++ b/Engine/Plugins/Developer/Concert/ConcertSync/ConcertSyncServer/Source/ConcertSyncServer/Private/Replication/MuteManager.cpp
@@ -63,10 +63,82 @@ namespace UE::ConcertSyncServer::Replication
 			: TOptional<FConcertReplication_ObjectMuteSetting>{};
 	}
 
+	bool FMuteManager::ValidateRequest(
+		const FConcertReplication_ChangeMuteState_Request& Request,
+		const ConcertSyncCore::FReplicatedObjectHierarchyCache* OverrideServerObjectCache,
+		TFunctionRef<void(const FSoftObjectPath& ObjectPath)> OnRejection
+		) const
+	{
+		const bool bFeatureSupported = EnumHasAnyFlags(SessionFlags, EConcertSyncSessionFlags::ShouldAllowGlobalMuting);
+		if (!ensureMsgf(bFeatureSupported, TEXT("Don't call this if the feature is not enabled!")))
+		{
+			return false;
+		}
+		
+		bool bIsValidRequest = true;
+		const auto Reject = [&OnRejection, &bIsValidRequest](const FSoftObjectPath& ObjectPath){ OnRejection(ObjectPath); bIsValidRequest = false; };
+		const ConcertSyncCore::FReplicatedObjectHierarchyCache& UsedObjectCache = OverrideServerObjectCache? *OverrideServerObjectCache : ServerObjectCache;
+		
+		const auto SharedValidateSetting = [this, &Reject, &UsedObjectCache](const FSoftObjectPath& Object, const FConcertReplication_ObjectMuteSetting& Setting)
+		{
+			const bool bIsObjectReferenced = UsedObjectCache.IsObjectReferencedDirectly(Object);
+			const bool bAppliesToSubobjects = ConcertSyncCore::AffectSubobjects(Setting.Flags);
+			if (!bIsObjectReferenced && !bAppliesToSubobjects)
+			{
+				Reject(Object);
+				return;
+			}
+
+			const bool bHasSubobjects = UsedObjectCache.HasChildren(Object);
+			if (!bIsObjectReferenced && bAppliesToSubobjects && !bHasSubobjects)
+			{
+				Reject(Object);
+			}
+		};
+		
+		for (const TPair<FSoftObjectPath, FConcertReplication_ObjectMuteSetting>& ToMute : Request.ObjectsToMute)
+		{
+			SharedValidateSetting(ToMute.Key, ToMute.Value);
+
+			// Cannot mute & unmute at the same time
+			if (Request.ObjectsToUnmute.Contains(ToMute.Key))
+			{
+				Reject(ToMute.Key);
+			}
+		}
+		for (const TPair<FSoftObjectPath, FConcertReplication_ObjectMuteSetting>& ToUnmute : Request.ObjectsToUnmute)
+		{
+			// If something is muted, the user can always unmute it. If it is already unmuted, it's still valid (but a no-op).
+			if (!MuteState.Contains(ToUnmute.Key))
+			{
+				SharedValidateSetting(ToUnmute.Key, ToUnmute.Value);
+			}
+		}
+		
+		return bIsValidRequest;
+	}
+
+	bool FMuteManager::ApplyRequestAndEnumerateSyncControl(
+		const FConcertReplication_ChangeMuteState_Request& Request,
+		const FOnSyncControlChange& OnSyncControlChange
+		)
+	{
+		const bool bIsValid = ValidateRequest(Request);
+		if (ensureMsgf(bIsValid, TEXT("This is a double-check and it failed. The caller should have validate the request beforehand.")))
+		{
+			InternalApplyRequest(Request);
+
+			check(OnRefreshSyncControlButSkipSendingToClientsDelegate.IsBound());
+			OnRefreshSyncControlButSkipSendingToClientsDelegate.Execute(OnSyncControlChange);
+		}
+		
+		return bIsValid;
+	}
+
 	FConcertReplication_ChangeSyncControl FMuteManager::ApplyManualRequest(const FGuid& EndpointId, const FConcertReplication_ChangeMuteState_Request& Request)
 	{
-		ApplyRequest(Request);
-		return OnGenerateSyncControlForMuteChangeDelegate.Execute(EndpointId);
+		InternalApplyRequest(Request);
+		return OnRefreshSyncControlAndSendToAllClientsExceptDelegate.Execute(EndpointId);
 	}
 
 	void FMuteManager::PostApplyStreamChange(const FGuid& ClientId, TConstArrayView<FConcertObjectInStreamID> AddedObjects, TConstArrayView<FConcertObjectInStreamID> RemovedObjects)
@@ -98,7 +170,7 @@ namespace UE::ConcertSyncServer::Replication
 			}
 		}
 	}
-	
+
 	void FMuteManager::TrackAddedSubobjectsForImplicitMuting(TConstArrayView<FConcertObjectInStreamID> AddedObjects)
 	{
 		for (const FConcertObjectInStreamID& Object : AddedObjects)
@@ -226,20 +298,20 @@ namespace UE::ConcertSyncServer::Replication
 			return EConcertSessionResponseCode::Failed;
 		}
 
-		if (!ValidateRequest(Request, Response))
+		if (!InternalValidateRequest(Request, Response))
 		{
 			Response.ErrorCode = EConcertReplicationMuteErrorCode::Rejected;
 			return EConcertSessionResponseCode::Success;
 		}
 
 		Response.ErrorCode = EConcertReplicationMuteErrorCode::Accepted;
-		ApplyRequest(Request);
+		InternalApplyRequest(Request);
 		
 		// Sync control will now 1. generate a sync control we can embed in the response, and 2. send a network event to all other clients.
 		// The embedded sync control will contain only those objects which changed sync control
 		// Fyi, the new sync control may contain more objects than were specified in the request since parent objects can also cause their subobjects' mute state to change.
-		check(OnGenerateSyncControlForMuteChangeDelegate.IsBound());
-		Response.SyncControl = OnGenerateSyncControlForMuteChangeDelegate.Execute(Context.SourceEndpointId);
+		check(OnRefreshSyncControlAndSendToAllClientsExceptDelegate.IsBound());
+		Response.SyncControl = OnRefreshSyncControlAndSendToAllClientsExceptDelegate.Execute(Context.SourceEndpointId);
 #if UE_BUILD_DEBUG // In debug builds so we don't tank performance
 		checkf(!Response.SyncControl.DoesAtLeastOneObjectLoseSyncControl(), TEXT("Clients are supposed to infer that objects lose sync control. We only tell them which objects gained sync control."));
 #endif
@@ -248,50 +320,13 @@ namespace UE::ConcertSyncServer::Replication
 		return EConcertSessionResponseCode::Success;
 	}
 
-	bool FMuteManager::ValidateRequest(const FConcertReplication_ChangeMuteState_Request& Request, FConcertReplication_ChangeMuteState_Response& Response) const
+	bool FMuteManager::InternalValidateRequest(const FConcertReplication_ChangeMuteState_Request& Request, FConcertReplication_ChangeMuteState_Response& Response) const
 	{
-		bool bIsValidRequest = true;
-		const auto Reject = [&Response, &bIsValidRequest](const FSoftObjectPath& ObjectPath){ Response.RejectionReasons.Add(ObjectPath); bIsValidRequest = false; };
-		const auto SharedValidateSetting = [this, &Reject](const FSoftObjectPath& Object, const FConcertReplication_ObjectMuteSetting& Setting)
-		{
-			const bool bIsObjectReferenced = ServerObjectCache.IsObjectReferencedDirectly(Object);
-			const bool bAppliesToSubobjects = ConcertSyncCore::AffectSubobjects(Setting.Flags);
-			if (!bIsObjectReferenced && !bAppliesToSubobjects)
-			{
-				Reject(Object);
-				return;
-			}
-
-			const bool bHasSubobjects = ServerObjectCache.HasChildren(Object);
-			if (!bIsObjectReferenced && bAppliesToSubobjects && !bHasSubobjects)
-			{
-				Reject(Object);
-			}
-		};
-		
-		for (const TPair<FSoftObjectPath, FConcertReplication_ObjectMuteSetting>& ToMute : Request.ObjectsToMute)
-		{
-			SharedValidateSetting(ToMute.Key, ToMute.Value);
-
-			// Cannot mute & unmute at the same time
-			if (Request.ObjectsToUnmute.Contains(ToMute.Key))
-			{
-				Reject(ToMute.Key);
-			}
-		}
-		for (const TPair<FSoftObjectPath, FConcertReplication_ObjectMuteSetting>& ToUnmute : Request.ObjectsToUnmute)
-		{
-			// If something is muted, the user can always unmute it. If it is already unmuted, it's still valid (but a no-op).
-			if (!MuteState.Contains(ToUnmute.Key))
-			{
-				SharedValidateSetting(ToUnmute.Key, ToUnmute.Value);
-			}
-		}
-		
-		return bIsValidRequest;
+		const auto Reject = [&Response](const FSoftObjectPath& ObjectPath){ Response.RejectionReasons.Add(ObjectPath); };
+		return ValidateRequest(Request, &ServerObjectCache, Reject);
 	}
 
-	void FMuteManager::ApplyRequest(const FConcertReplication_ChangeMuteState_Request& Request)
+	void FMuteManager::InternalApplyRequest(const FConcertReplication_ChangeMuteState_Request& Request)
 	{
 		UE_LOG(LogConcert, Log, TEXT("Explicitly muting %d and unmuting %d objects."), Request.ObjectsToMute.Num(), Request.ObjectsToUnmute.Num());
 		
diff --git a/Engine/Plugins/Developer/Concert/ConcertSync/ConcertSyncServer/Source/ConcertSyncServer/Private/Replication/MuteManager.h b/Engine/Plugins/Developer/Concert/ConcertSync/ConcertSyncServer/Source/ConcertSyncServer/Private/Replication/MuteManager.h
--- a/Engine/Plugins/Developer/Concert/ConcertSync/ConcertSyncServer/Source/ConcertSyncServer/Private/Replication/MuteManager.h
+++ b/Engine/Plugins/Developer/Concert/ConcertSync/ConcertSyncServer/Source/ConcertSyncServer/Private/Replication/MuteManager.h
@@ -8,6 +8,7 @@
 #include "Replication/Messages/Muting.h"
 
 #include "Misc/Optional.h"
+#include "Templates/Function.h"
 #include "Templates/UnrealTemplate.h"
 
 #include <type_traits>
@@ -41,10 +42,11 @@ namespace UE::ConcertSyncServer::Replication
 	{
 	public:
 
-		/** Updates sync control if needed in response to a client changing their mute state indirectly (e.g. because client removed object from their stream). */
+		using FOnSyncControlChange = TFunctionRef<void(const FGuid& ClientId, const FConcertReplication_ChangeSyncControl& SyncControlChange)>;
+
 		DECLARE_DELEGATE_OneParam(FOnMuteStateChangedByClient, const FGuid& ClientId);
-		/** Updates sync control for all clients, sends an update to all clients but ClientId, and returns the sync control to embed into the mute response. */
-		DECLARE_DELEGATE_RetVal_OneParam(FConcertReplication_ChangeSyncControl, FGenerateSyncControlForMuteChange, const FGuid& ClientId);
+		DECLARE_DELEGATE_RetVal_OneParam(FConcertReplication_ChangeSyncControl, FRefreshSyncControlForMuteChange, const FGuid& ClientId);
+		DECLARE_DELEGATE_OneParam(FRefreshSyncControlForClients, const FOnSyncControlChange& OnSyncControlChange);
 		
 		/** Notifies manager about applied event. Generates mute activity. */
 		DECLARE_MULTICAST_DELEGATE_TwoParams(FOnMuteRequestApplied, const FGuid& ClientId, const FConcertReplication_ChangeMuteState_Request& Request);
@@ -76,7 +78,30 @@ namespace UE::ConcertSyncServer::Replication
 		TOptional<EMuteState> GetMuteState(const FSoftObjectPath& Object) const;
 		/** @return The mute setting of Object, if it is explicitly set (i.e. not affected by a parent object). */
 		TOptional<FConcertReplication_ObjectMuteSetting> GetExplicitMuteSetting(const FSoftObjectPath& Object) const;
-
+		
+		/**
+		 * Checks whether Request would be valid to apply, optionally providing a future server state.
+		 * @param Request The request to validate
+		 * @param OverrideServerObjectCache The server cache to use. Defaults to the internal one if left nullptr. Useful for passing in a future server state.
+		 * @param OnRejection Called for every rejected object.
+		 */
+		bool ValidateRequest(
+			const FConcertReplication_ChangeMuteState_Request& Request,
+			const ConcertSyncCore::FReplicatedObjectHierarchyCache* OverrideServerObjectCache = nullptr,
+			TFunctionRef<void(const FSoftObjectPath& ObjectPath)> OnRejection = [](const FSoftObjectPath&){}
+			) const;
+		/**
+		 * This validates and then applies Request.
+		 * This version does not send any sync control updates to other clients. You must update them yourself.
+		 * 
+		 * @param Request The request to apply
+		 * @param OnSyncControlChange Called for every sync control change that is made.
+		 * @return Whether the request was applied.
+		 */
+		bool ApplyRequestAndEnumerateSyncControl(
+			const FConcertReplication_ChangeMuteState_Request& Request,
+			const FOnSyncControlChange& OnSyncControlChange
+			);
 		/**
 		 * Applies Request as if it was sent by EndpointId.
 		 * @return The sync control that the client identified by EndpointId has gained. Does NOT include the sync control they lost.
@@ -88,8 +113,12 @@ namespace UE::ConcertSyncServer::Replication
 		/** Called by FConcertServerReplicationManager when client leaves replication. */
 		void OnPostClientLeft(const TArray<FConcertReplicationStream>& ClientStreams);
 
-		FOnMuteStateChangedByClient& OnUpdateSyncControlForIndirectMuteChange() { return OnMuteStateChangedDelegate; }
-		FGenerateSyncControlForMuteChange& OnGenerateSyncControlForMuteChange() { return OnGenerateSyncControlForMuteChangeDelegate; }
+		/** Updates sync control if needed in response to a client changing their mute state indirectly (e.g. because client removed object from their stream). */
+		FOnMuteStateChangedByClient& OnRefreshSyncControlForIndirectMuteChange() { return OnMuteStateChangedDelegate; }
+		/** Updates sync control for all clients, sends an update to all clients but ClientId, and returns the sync control to embed into the mute response. */
+		FRefreshSyncControlForMuteChange& OnRefreshSyncControlAndSendToAllClientsExcept() { return OnRefreshSyncControlAndSendToAllClientsExceptDelegate; }
+		/** Updates sync control for all clients but does not send any update to any clients. The passed in callback receives all sync control changes for the clients. */
+		FRefreshSyncControlForClients& OnRefreshSyncControlButSkipSendingToClients() { return OnRefreshSyncControlButSkipSendingToClientsDelegate; }
 		
 		FOnMuteRequestApplied& OnMuteRequestApplied() { return OnMuteRequestAppliedDelegate; }
 		
@@ -111,7 +140,9 @@ namespace UE::ConcertSyncServer::Replication
 		 * Updates sync control for all clients, sends an update to all clients but ClientId, and returns the sync control to embed into the mute response.
 		 * The sync control will only contain RestrictToObjects.
 		 */
-		FGenerateSyncControlForMuteChange OnGenerateSyncControlForMuteChangeDelegate;
+		FRefreshSyncControlForMuteChange OnRefreshSyncControlAndSendToAllClientsExceptDelegate;
+		/** Delegate into FSyncControlManager.  */
+		FRefreshSyncControlForClients OnRefreshSyncControlButSkipSendingToClientsDelegate;
 		
 		/** Broadcasts after a mute request has been applied. */
 		FOnMuteRequestApplied OnMuteRequestAppliedDelegate;
@@ -151,9 +182,9 @@ namespace UE::ConcertSyncServer::Replication
 		/** Handles requests to change the mute states */
 		EConcertSessionResponseCode HandleChangeMuteStateRequest(const FConcertSessionContext& Context, const FConcertReplication_ChangeMuteState_Request& Request, FConcertReplication_ChangeMuteState_Response& Response);
 		/** Checks whether Request is valid to apply. */
-		bool ValidateRequest(const FConcertReplication_ChangeMuteState_Request& Request, FConcertReplication_ChangeMuteState_Response& Response) const;
+		bool InternalValidateRequest(const FConcertReplication_ChangeMuteState_Request& Request, FConcertReplication_ChangeMuteState_Response& Response) const;
 		/** Updates the internal state (assuming it is a valid request - call ValidateRequest before). */
-		void ApplyRequest(const FConcertReplication_ChangeMuteState_Request& Request);
+		void InternalApplyRequest(const FConcertReplication_ChangeMuteState_Request& Request);
 
 		/** Removes all mute state for Object and transitively updates subobjects if applicable. */
 		template<CObjectProcessable TCallback>
diff --git a/Engine/Plugins/Developer/Concert/ConcertSync/ConcertSyncServer/Source/ConcertSyncServer/Private/Replication/SyncControlManager.cpp b/Engine/Plugins/Developer/Concert/ConcertSync/ConcertSyncServer/Source/ConcertSyncServer/Private/Replication/SyncControlManager.cpp
--- a/Engine/Plugins/Developer/Concert/ConcertSync/ConcertSyncServer/Source/ConcertSyncServer/Private/Replication/SyncControlManager.cpp
+++ b/Engine/Plugins/Developer/Concert/ConcertSync/ConcertSyncServer/Source/ConcertSyncServer/Private/Replication/SyncControlManager.cpp
@@ -22,8 +22,9 @@ namespace UE::ConcertSyncServer::Replication
 		, Getters(Getters)
 	{
 		AuthorityManager.OnGenerateSyncControl().BindRaw(this, &FSyncControlManager::OnGenerateSyncControlForAuthorityResponse);
-		MuteManager.OnUpdateSyncControlForIndirectMuteChange().BindRaw(this, &FSyncControlManager::OnUpdateSyncControlForIndirectMuteChange);
-		MuteManager.OnGenerateSyncControlForMuteChange().BindRaw(this, &FSyncControlManager::OnGenerateSyncControlForClientMuteChange);
+		MuteManager.OnRefreshSyncControlForIndirectMuteChange().BindRaw(this, &FSyncControlManager::OnRefreshSyncControlForIndirectMuteChange);
+		MuteManager.OnRefreshSyncControlAndSendToAllClientsExcept().BindRaw(this, &FSyncControlManager::OnRefreshSyncControlForClientMuteChange);
+		MuteManager.OnRefreshSyncControlButSkipSendingToClients().BindRaw(this, &FSyncControlManager::OnRefreshSyncControlAndEnumerateWithoutSending);
 	}
 
 	FSyncControlManager::~FSyncControlManager()
@@ -34,13 +35,17 @@ namespace UE::ConcertSyncServer::Replication
 		{
 			AuthorityManager.OnGenerateSyncControl().Unbind();
 		}
-		if (ensure(MuteManager.OnUpdateSyncControlForIndirectMuteChange().IsBoundToObject(this)))
+		if (ensure(MuteManager.OnRefreshSyncControlForIndirectMuteChange().IsBoundToObject(this)))
 		{
-			MuteManager.OnUpdateSyncControlForIndirectMuteChange().Unbind();
+			MuteManager.OnRefreshSyncControlForIndirectMuteChange().Unbind();
 		}
-		if (ensure(MuteManager.OnGenerateSyncControlForMuteChange().IsBoundToObject(this)))
+		if (ensure(MuteManager.OnRefreshSyncControlAndSendToAllClientsExcept().IsBoundToObject(this)))
 		{
-			MuteManager.OnGenerateSyncControlForMuteChange().Unbind();
+			MuteManager.OnRefreshSyncControlAndSendToAllClientsExcept().Unbind();
+		}
+		if (ensure(MuteManager.OnRefreshSyncControlButSkipSendingToClients().IsBoundToObject(this)))
+		{
+			MuteManager.OnRefreshSyncControlButSkipSendingToClients().Unbind();
 		}
 	}
 
@@ -69,7 +74,7 @@ namespace UE::ConcertSyncServer::Replication
 		return RefreshClientSyncControl(ClientId, ShouldSkipInControlMessage);
 	}
 
-	FConcertReplication_ChangeSyncControl FSyncControlManager::OnGenerateSyncControlForClientMuteChange(const FGuid& ClientId)
+	FConcertReplication_ChangeSyncControl FSyncControlManager::OnRefreshSyncControlForClientMuteChange(const FGuid& ClientId)
 	{
 		RefreshAndSendToAllClientsExcept(ClientId);
 		const auto ShouldSkipInControlMessage = [](const FConcertObjectInStreamID&, bool bNewState)
@@ -82,6 +87,15 @@ namespace UE::ConcertSyncServer::Replication
 		return RefreshClientSyncControl(ClientId, ShouldSkipInControlMessage);
 	}
 
+	void FSyncControlManager::OnRefreshSyncControlAndEnumerateWithoutSending(const FMuteManager::FOnSyncControlChange& OnSyncControlChange)
+	{
+		Getters.ForEachReplicationClient([this, &OnSyncControlChange](const FGuid& ClientId)
+		{
+			OnSyncControlChange(ClientId, RefreshClientSyncControl(ClientId));
+			return EBreakBehavior::Continue;
+		});
+	}
+
 	void FSyncControlManager::HandleClientLeave(const FGuid& LeftClientId)
 	{
 		// The disconnecting client's objects should be removed
diff --git a/Engine/Plugins/Developer/Concert/ConcertSync/ConcertSyncServer/Source/ConcertSyncServer/Private/Replication/SyncControlManager.h b/Engine/Plugins/Developer/Concert/ConcertSync/ConcertSyncServer/Source/ConcertSyncServer/Private/Replication/SyncControlManager.h
--- a/Engine/Plugins/Developer/Concert/ConcertSync/ConcertSyncServer/Source/ConcertSyncServer/Private/Replication/SyncControlManager.h
+++ b/Engine/Plugins/Developer/Concert/ConcertSync/ConcertSyncServer/Source/ConcertSyncServer/Private/Replication/SyncControlManager.h
@@ -2,6 +2,7 @@
 
 #pragma once
 
+#include "MuteManager.h"
 #include "Replication/Data/ObjectIds.h"
 
 #include "HAL/Platform.h"
@@ -77,7 +78,7 @@ namespace UE::ConcertSyncServer::Replication
 		 * Called when a client has caused mute state to change (e.g. due to removing object from stream, or explicitly muting it).
 		 * Called via delegate by FMuteManager.
 		 */
-		void OnUpdateSyncControlForIndirectMuteChange(const FGuid& ClientId)
+		void OnRefreshSyncControlForIndirectMuteChange(const FGuid& ClientId)
 		{
 			// Do not send any update to this client (because they can infer the change themselves) but do update the sync control state.
 			RefreshClientSyncControl(ClientId);
@@ -88,7 +89,9 @@ namespace UE::ConcertSyncServer::Replication
 		 * Updates sync control for all clients, sends an update to all clients but ClientId, and returns the sync control to embed into the mute response.
 		 * Called via delegate by FMuteManager.
 		 */
-		FConcertReplication_ChangeSyncControl OnGenerateSyncControlForClientMuteChange(const FGuid& ClientId);
+		FConcertReplication_ChangeSyncControl OnRefreshSyncControlForClientMuteChange(const FGuid& ClientId);
+		/** Updates sync control for all clients and enumerates the sync control instead of sending it. The callback handles updating the remote client machines. */
+		void OnRefreshSyncControlAndEnumerateWithoutSending(const FMuteManager::FOnSyncControlChange& OnSyncControlChange);
 		
 		/** Cleans up the associated client data and updates sync control for all other clients. */
 		void HandleClientLeave(const FGuid& LeftClientId);
