Introduce IMuteValidationObjectHierarchy which is passed to FMuteManager::ValidateRequest to validate requests against future server state.

diff --git a/Engine/Plugins/Developer/Concert/ConcertSync/ConcertSyncCore/Source/ConcertSyncCore/Private/Misc/ObjectPathHierarchy.cpp b/Engine/Plugins/Developer/Concert/ConcertSync/ConcertSyncCore/Source/ConcertSyncCore/Private/Misc/ObjectPathHierarchy.cpp
--- a/Engine/Plugins/Developer/Concert/ConcertSync/ConcertSyncCore/Source/ConcertSyncCore/Private/Misc/ObjectPathHierarchy.cpp
+++ b/Engine/Plugins/Developer/Concert/ConcertSync/ConcertSyncCore/Source/ConcertSyncCore/Private/Misc/ObjectPathHierarchy.cpp
@@ -18,6 +18,23 @@ namespace UE::ConcertSyncCore
 		}
 	}
 
+	FObjectPathHierarchy::FObjectPathHierarchy(FObjectPathHierarchy&& Other)
+		: AssetNodes(MoveTemp(Other.AssetNodes))
+		, CachedNodes(MoveTemp(Other.CachedNodes))
+	{}
+
+	FObjectPathHierarchy& FObjectPathHierarchy::operator=(FObjectPathHierarchy&& Other)
+	{
+		if (this == &Other)
+		{
+			return *this;
+		}
+		
+		AssetNodes = MoveTemp(Other.AssetNodes);
+		CachedNodes = MoveTemp(Other.CachedNodes);
+		return *this;
+	}
+
 	void FObjectPathHierarchy::TraverseTopToBottom(TFunctionRef<ETreeTraversalBehavior(const FChildRelation& Relation)> Callback, const FSoftObjectPath& Start) const
 	{
 		const bool bHasStart = !Start.IsNull();
diff --git a/Engine/Plugins/Developer/Concert/ConcertSync/ConcertSyncCore/Source/ConcertSyncCore/Public/Misc/ObjectPathHierarchy.h b/Engine/Plugins/Developer/Concert/ConcertSync/ConcertSyncCore/Source/ConcertSyncCore/Public/Misc/ObjectPathHierarchy.h
--- a/Engine/Plugins/Developer/Concert/ConcertSync/ConcertSyncCore/Source/ConcertSyncCore/Public/Misc/ObjectPathHierarchy.h
+++ b/Engine/Plugins/Developer/Concert/ConcertSync/ConcertSyncCore/Source/ConcertSyncCore/Public/Misc/ObjectPathHierarchy.h
@@ -50,11 +50,15 @@ namespace UE::ConcertSyncCore
 	
 	/** Keeps track of the outer tree hierarchy of FSoftObjectPath. */
 	class CONCERTSYNCCORE_API FObjectPathHierarchy
-		// This class cannot be shallowly copied because the nodes referenced dynamic allocations - we don't expect instances to be copied, so we won't bother implementing a deep copy.
+		// This class cannot be shallowly copied because the nodes' referenced dynamic allocations - we don't expect instances to be copied, so we won't bother implementing a deep copy.
 		: public FNoncopyable
 	{
 	public:
-		
+
+		FObjectPathHierarchy() = default;
+		FObjectPathHierarchy(FObjectPathHierarchy&& Other);
+		FObjectPathHierarchy& operator=(FObjectPathHierarchy&& Other);
+
 		/**
 		 * Traverses the hierarchy in pre-order (root first, then its children), starting at an optional parent object.
 		 *
diff --git a/Engine/Plugins/Developer/Concert/ConcertSync/ConcertSyncServer/Source/ConcertSyncServer/Private/Replication/MuteManager.cpp b/Engine/Plugins/Developer/Concert/ConcertSync/ConcertSyncServer/Source/ConcertSyncServer/Private/Replication/MuteManager.cpp
--- a/Engine/Plugins/Developer/Concert/ConcertSync/ConcertSyncServer/Source/ConcertSyncServer/Private/Replication/MuteManager.cpp
+++ b/Engine/Plugins/Developer/Concert/ConcertSync/ConcertSyncServer/Source/ConcertSyncServer/Private/Replication/MuteManager.cpp
@@ -12,6 +12,7 @@
 
 #include "HAL/IConsoleManager.h"
 #include "Misc/ScopeExit.h"
+#include "Muting/ObjectHierarchyAdapter.h"
 
 namespace UE::ConcertSyncServer::Replication
 {
@@ -65,7 +66,7 @@ namespace UE::ConcertSyncServer::Replication
 
 	bool FMuteManager::ValidateRequest(
 		const FConcertReplication_ChangeMuteState_Request& Request,
-		const ConcertSyncCore::FReplicatedObjectHierarchyCache* OverrideServerObjectCache,
+		const IMuteValidationObjectHierarchy* OverrideServerObjectCache,
 		TFunctionRef<void(const FSoftObjectPath& ObjectPath)> OnRejection
 		) const
 	{
@@ -74,48 +75,10 @@ namespace UE::ConcertSyncServer::Replication
 		{
 			return false;
 		}
-		
-		bool bIsValidRequest = true;
-		const auto Reject = [&OnRejection, &bIsValidRequest](const FSoftObjectPath& ObjectPath){ OnRejection(ObjectPath); bIsValidRequest = false; };
-		const ConcertSyncCore::FReplicatedObjectHierarchyCache& UsedObjectCache = OverrideServerObjectCache? *OverrideServerObjectCache : ServerObjectCache;
-		
-		const auto SharedValidateSetting = [this, &Reject, &UsedObjectCache](const FSoftObjectPath& Object, const FConcertReplication_ObjectMuteSetting& Setting)
-		{
-			const bool bIsObjectReferenced = UsedObjectCache.IsObjectReferencedDirectly(Object);
-			const bool bAppliesToSubobjects = ConcertSyncCore::AffectSubobjects(Setting.Flags);
-			if (!bIsObjectReferenced && !bAppliesToSubobjects)
-			{
-				Reject(Object);
-				return;
-			}
-
-			const bool bHasSubobjects = UsedObjectCache.HasChildren(Object);
-			if (!bIsObjectReferenced && bAppliesToSubobjects && !bHasSubobjects)
-			{
-				Reject(Object);
-			}
-		};
-		
-		for (const TPair<FSoftObjectPath, FConcertReplication_ObjectMuteSetting>& ToMute : Request.ObjectsToMute)
-		{
-			SharedValidateSetting(ToMute.Key, ToMute.Value);
 
-			// Cannot mute & unmute at the same time
-			if (Request.ObjectsToUnmute.Contains(ToMute.Key))
-			{
-				Reject(ToMute.Key);
-			}
-		}
-		for (const TPair<FSoftObjectPath, FConcertReplication_ObjectMuteSetting>& ToUnmute : Request.ObjectsToUnmute)
-		{
-			// If something is muted, the user can always unmute it. If it is already unmuted, it's still valid (but a no-op).
-			if (!MuteState.Contains(ToUnmute.Key))
-			{
-				SharedValidateSetting(ToUnmute.Key, ToUnmute.Value);
-			}
-		}
-		
-		return bIsValidRequest;
+		return OverrideServerObjectCache
+			? ValidateRequestInternal(Request, *OverrideServerObjectCache, OnRejection)
+			: ValidateRequestInternal(Request, FObjectHierarchyAdapter(ServerObjectCache), OnRejection);
 	}
 
 	bool FMuteManager::ApplyRequestAndEnumerateSyncControl(
@@ -126,7 +89,7 @@ namespace UE::ConcertSyncServer::Replication
 		const bool bIsValid = ValidateRequest(Request);
 		if (ensureMsgf(bIsValid, TEXT("This is a double-check and it failed. The caller should have validate the request beforehand.")))
 		{
-			InternalApplyRequest(Request);
+			ApplyRequestInternal(Request);
 
 			check(OnRefreshSyncControlButSkipSendingToClientsDelegate.IsBound());
 			OnRefreshSyncControlButSkipSendingToClientsDelegate.Execute(OnSyncControlChange);
@@ -137,7 +100,7 @@ namespace UE::ConcertSyncServer::Replication
 
 	FConcertReplication_ChangeSyncControl FMuteManager::ApplyManualRequest(const FGuid& EndpointId, const FConcertReplication_ChangeMuteState_Request& Request)
 	{
-		InternalApplyRequest(Request);
+		ApplyRequestInternal(Request);
 		return OnRefreshSyncControlAndSendToAllClientsExceptDelegate.Execute(EndpointId);
 	}
 
@@ -298,14 +261,14 @@ namespace UE::ConcertSyncServer::Replication
 			return EConcertSessionResponseCode::Failed;
 		}
 
-		if (!InternalValidateRequest(Request, Response))
+		if (!ValidateRequestInternal(Request, Response))
 		{
 			Response.ErrorCode = EConcertReplicationMuteErrorCode::Rejected;
 			return EConcertSessionResponseCode::Success;
 		}
 
 		Response.ErrorCode = EConcertReplicationMuteErrorCode::Accepted;
-		InternalApplyRequest(Request);
+		ApplyRequestInternal(Request);
 		
 		// Sync control will now 1. generate a sync control we can embed in the response, and 2. send a network event to all other clients.
 		// The embedded sync control will contain only those objects which changed sync control
@@ -320,13 +283,61 @@ namespace UE::ConcertSyncServer::Replication
 		return EConcertSessionResponseCode::Success;
 	}
 
-	bool FMuteManager::InternalValidateRequest(const FConcertReplication_ChangeMuteState_Request& Request, FConcertReplication_ChangeMuteState_Response& Response) const
+	bool FMuteManager::ValidateRequestInternal(const FConcertReplication_ChangeMuteState_Request& Request, FConcertReplication_ChangeMuteState_Response& Response) const
 	{
 		const auto Reject = [&Response](const FSoftObjectPath& ObjectPath){ Response.RejectionReasons.Add(ObjectPath); };
-		return ValidateRequest(Request, &ServerObjectCache, Reject);
+		return ValidateRequestInternal(Request, FObjectHierarchyAdapter(ServerObjectCache), Reject);
+	}
+
+	bool FMuteManager::ValidateRequestInternal(
+		const FConcertReplication_ChangeMuteState_Request& Request, 
+		const IMuteValidationObjectHierarchy& ObjectCache,
+		TFunctionRef<void(const FSoftObjectPath& ObjectPath)> OnRejection
+		) const
+	{
+		bool bIsValidRequest = true;
+		const auto Reject = [&OnRejection, &bIsValidRequest](const FSoftObjectPath& ObjectPath){ OnRejection(ObjectPath); bIsValidRequest = false; };
+		
+		const auto SharedValidateSetting = [this, &Reject, &ObjectCache](const FSoftObjectPath& Object, const FConcertReplication_ObjectMuteSetting& Setting)
+		{
+			const bool bIsObjectReferenced = ObjectCache.IsObjectReferencedDirectly(Object);
+			const bool bAppliesToSubobjects = ConcertSyncCore::AffectSubobjects(Setting.Flags);
+			if (!bIsObjectReferenced && !bAppliesToSubobjects)
+			{
+				Reject(Object);
+				return;
+			}
+
+			const bool bHasSubobjects = ObjectCache.HasChildren(Object);
+			if (!bIsObjectReferenced && bAppliesToSubobjects && !bHasSubobjects)
+			{
+				Reject(Object);
+			}
+		};
+		
+		for (const TPair<FSoftObjectPath, FConcertReplication_ObjectMuteSetting>& ToMute : Request.ObjectsToMute)
+		{
+			SharedValidateSetting(ToMute.Key, ToMute.Value);
+
+			// Cannot mute & unmute at the same time
+			if (Request.ObjectsToUnmute.Contains(ToMute.Key))
+			{
+				Reject(ToMute.Key);
+			}
+		}
+		for (const TPair<FSoftObjectPath, FConcertReplication_ObjectMuteSetting>& ToUnmute : Request.ObjectsToUnmute)
+		{
+			// If something is muted, the user can always unmute it. If it is already unmuted, it's still valid (but a no-op).
+			if (!MuteState.Contains(ToUnmute.Key))
+			{
+				SharedValidateSetting(ToUnmute.Key, ToUnmute.Value);
+			}
+		}
+		
+		return bIsValidRequest;
 	}
 
-	void FMuteManager::InternalApplyRequest(const FConcertReplication_ChangeMuteState_Request& Request)
+	void FMuteManager::ApplyRequestInternal(const FConcertReplication_ChangeMuteState_Request& Request)
 	{
 		UE_LOG(LogConcert, Log, TEXT("Explicitly muting %d and unmuting %d objects."), Request.ObjectsToMute.Num(), Request.ObjectsToUnmute.Num());
 		
diff --git a/Engine/Plugins/Developer/Concert/ConcertSync/ConcertSyncServer/Source/ConcertSyncServer/Private/Replication/MuteManager.h b/Engine/Plugins/Developer/Concert/ConcertSync/ConcertSyncServer/Source/ConcertSyncServer/Private/Replication/MuteManager.h
--- a/Engine/Plugins/Developer/Concert/ConcertSync/ConcertSyncServer/Source/ConcertSyncServer/Private/Replication/MuteManager.h
+++ b/Engine/Plugins/Developer/Concert/ConcertSync/ConcertSyncServer/Source/ConcertSyncServer/Private/Replication/MuteManager.h
@@ -19,7 +19,9 @@ struct FConcertReplicationStream;
 struct FConcertObjectInStreamID;
 
 namespace UE::ConcertSyncCore { class FReplicatedObjectHierarchyCache; }
-namespace UE::ConcertSyncServer::Replication { class IRegistrationEnumerator; }
+namespace UE::ConcertSyncServer::Replication {
+	class IMuteValidationObjectHierarchy;
+	class IRegistrationEnumerator; }
 
 namespace UE::ConcertSyncServer::Replication
 {
@@ -87,9 +89,10 @@ namespace UE::ConcertSyncServer::Replication
 		 */
 		bool ValidateRequest(
 			const FConcertReplication_ChangeMuteState_Request& Request,
-			const ConcertSyncCore::FReplicatedObjectHierarchyCache* OverrideServerObjectCache = nullptr,
+			const IMuteValidationObjectHierarchy* OverrideServerObjectCache = nullptr,
 			TFunctionRef<void(const FSoftObjectPath& ObjectPath)> OnRejection = [](const FSoftObjectPath&){}
 			) const;
+		
 		/**
 		 * This validates and then applies Request.
 		 * This version does not send any sync control updates to other clients. You must update them yourself.
@@ -181,10 +184,13 @@ namespace UE::ConcertSyncServer::Replication
 		EConcertSessionResponseCode HandleQueryMuteStateRequest(const FConcertSessionContext& Context, const FConcertReplication_QueryMuteState_Request& Request, FConcertReplication_QueryMuteState_Response& Response);
 		/** Handles requests to change the mute states */
 		EConcertSessionResponseCode HandleChangeMuteStateRequest(const FConcertSessionContext& Context, const FConcertReplication_ChangeMuteState_Request& Request, FConcertReplication_ChangeMuteState_Response& Response);
+
+		/** Checks whether Request is valid to apply. */
+		bool ValidateRequestInternal(const FConcertReplication_ChangeMuteState_Request& Request, FConcertReplication_ChangeMuteState_Response& Response) const;
 		/** Checks whether Request is valid to apply. */
-		bool InternalValidateRequest(const FConcertReplication_ChangeMuteState_Request& Request, FConcertReplication_ChangeMuteState_Response& Response) const;
+		bool ValidateRequestInternal(const FConcertReplication_ChangeMuteState_Request& Request, const IMuteValidationObjectHierarchy& ObjectCache, TFunctionRef<void(const FSoftObjectPath& ObjectPath)> OnRejection) const;
 		/** Updates the internal state (assuming it is a valid request - call ValidateRequest before). */
-		void InternalApplyRequest(const FConcertReplication_ChangeMuteState_Request& Request);
+		void ApplyRequestInternal(const FConcertReplication_ChangeMuteState_Request& Request);
 
 		/** Removes all mute state for Object and transitively updates subobjects if applicable. */
 		template<CObjectProcessable TCallback>
diff --git a/Engine/Plugins/Developer/Concert/ConcertSync/ConcertSyncServer/Source/ConcertSyncServer/Private/Replication/Muting/IMuteValidationObjectHierarchy.h b/Engine/Plugins/Developer/Concert/ConcertSync/ConcertSyncServer/Source/ConcertSyncServer/Private/Replication/Muting/IMuteValidationObjectHierarchy.h
new file mode 100644
--- /dev/null
+++ b/Engine/Plugins/Developer/Concert/ConcertSync/ConcertSyncServer/Source/ConcertSyncServer/Private/Replication/Muting/IMuteValidationObjectHierarchy.h
@@ -0,0 +1,45 @@
+﻿// Copyright Epic Games, Inc. All Rights Reserved.
+
+#pragma once
+
+#include "Containers/ArrayView.h"
+#include "Containers/ContainersFwd.h"
+
+struct FGuid;
+struct FSoftObjectPath;
+
+namespace UE::ConcertSyncServer::Replication
+{
+	/**
+	 * Interface required by FMuteManager to validate a FConcertReplication_ChangeMuteState_Request.
+	 * 
+	 * This encapsulates the server's knowledge of replicated objects.
+	 * It is useful if you have a complicated series of requests that will mutate the server state if applied, but you want to first validate that
+	 * a mute request is valid to apply on that state. In that case, you can implement this to return the future server state.
+	 */
+	class IMuteValidationObjectHierarchy
+	{
+	public:
+
+		/**
+		 * Checks whether the object is registered by any client (except for this in IgnoredClients).
+		 *
+		 * This function ignores implicit knowledge of the hierarchy.
+		 * For example if you register ONLY /Game/Maps.Map:PersistentLevel.Cube.StaticMeshComponent0, then
+		 * - IsObjectReferencedDirectly(/Game/Maps.Map:PersistentLevel.Cube.StaticMeshComponent0) == true,
+		 * - IsObjectReferencedDirectly(/Game/Maps.Map:PersistentLevel.Cube) == false
+		 * but e.g. TraverseTopToBottom would list both paths.
+		 *
+		 * @see FReplicatedObjectHierarchyCache::IsObjectReferencedDirectly
+		 */
+		virtual bool IsObjectReferencedDirectly(const FSoftObjectPath& ObjectPath, TConstArrayView<FGuid> IgnoredClients = {}) const = 0;
+
+		/**
+		 * @return Whether Object has any subobjects in the hierarchy.
+		 * @see FObjectPathHierarchy::HasChildren
+		 */
+		virtual bool HasChildren(const FSoftObjectPath& Object) const = 0;
+
+		virtual ~IMuteValidationObjectHierarchy() = default;
+	};
+}
\ No newline at end of file
diff --git a/Engine/Plugins/Developer/Concert/ConcertSync/ConcertSyncServer/Source/ConcertSyncServer/Private/Replication/Muting/ObjectHierarchyAdapter.h b/Engine/Plugins/Developer/Concert/ConcertSync/ConcertSyncServer/Source/ConcertSyncServer/Private/Replication/Muting/ObjectHierarchyAdapter.h
new file mode 100644
--- /dev/null
+++ b/Engine/Plugins/Developer/Concert/ConcertSync/ConcertSyncServer/Source/ConcertSyncServer/Private/Replication/Muting/ObjectHierarchyAdapter.h
@@ -0,0 +1,27 @@
+﻿// Copyright Epic Games, Inc. All Rights Reserved.
+
+#pragma once
+
+#include "IMuteValidationObjectHierarchy.h"
+#include "Replication/Misc/ReplicatedObjectHierarchyCache.h"
+
+#include "HAL/Platform.h"
+
+namespace UE::ConcertSyncServer::Replication
+{
+	/** Simple adapter implementation that uses the server's current FReplicatedObjectHierarchyCache state. */
+	class FObjectHierarchyAdapter : public IMuteValidationObjectHierarchy
+	{
+		const ConcertSyncCore::FReplicatedObjectHierarchyCache& ServerCache;
+	public:
+		
+		FObjectHierarchyAdapter(const ConcertSyncCore::FReplicatedObjectHierarchyCache& ServerCache UE_LIFETIMEBOUND)
+			: ServerCache(ServerCache)
+		{}
+
+		//~ Begin IMuteValidationObjectHierarchy Interface
+		virtual bool IsObjectReferencedDirectly(const FSoftObjectPath& ObjectPath, TConstArrayView<FGuid> IgnoredClients) const override { return ServerCache.IsObjectReferencedDirectly(ObjectPath, IgnoredClients); }
+		virtual bool HasChildren(const FSoftObjectPath& Object) const override { return ServerCache.HasChildren(Object); }
+		//~ Begin IMuteValidationObjectHierarchy Interface
+	};
+}
diff --git a/Engine/Plugins/Developer/Concert/ConcertSync/ConcertSyncServer/Source/ConcertSyncServer/Private/Replication/Muting/PredictedStateObjectHierarchy.cpp b/Engine/Plugins/Developer/Concert/ConcertSync/ConcertSyncServer/Source/ConcertSyncServer/Private/Replication/Muting/PredictedStateObjectHierarchy.cpp
new file mode 100644
--- /dev/null
+++ b/Engine/Plugins/Developer/Concert/ConcertSync/ConcertSyncServer/Source/ConcertSyncServer/Private/Replication/Muting/PredictedStateObjectHierarchy.cpp
@@ -0,0 +1,41 @@
+﻿// Copyright Epic Games, Inc. All Rights Reserved.
+
+#include "PredictedStateObjectHierarchy.h"
+
+#include "Algo/AllOf.h"
+#include "Replication/Data/ReplicationStream.h"
+
+namespace UE::ConcertSyncServer::Replication
+{
+	void FPredictedStateObjectHierarchy::AddClients(const TMap<FGuid, TArray<FConcertReplicationStream>>& Clients)
+	{
+		for (const TPair<FGuid, TArray<FConcertReplicationStream>>& Pair : Clients)
+		{
+			AddClientData(Pair.Key, Pair.Value);
+		}
+	}
+
+	void FPredictedStateObjectHierarchy::AddClientData(const FGuid& ClientId, TConstArrayView<FConcertReplicationStream> Streams)
+	{
+		for (const FConcertReplicationStream& Stream : Streams)
+		{
+			for (const TPair<FSoftObjectPath, FConcertReplicatedObjectInfo>& Pair : Stream.BaseDescription.ReplicationMap.ReplicatedObjects)
+			{
+				const FSoftObjectPath& ObjectPath = Pair.Key;
+				Hierarchy.AddObject(ObjectPath);
+				ObjectReferencingClients.FindOrAdd(ObjectPath).AddUnique(ClientId);
+			}
+		}
+	}
+
+	bool FPredictedStateObjectHierarchy::IsObjectReferencedDirectly(const FSoftObjectPath& ObjectPath, TConstArrayView<FGuid> IgnoredClients) const
+	{
+		const TArray<FGuid>* ReferencingClients = ObjectReferencingClients.Find(ObjectPath);
+		return ReferencingClients && Algo::AllOf(*ReferencingClients, [&IgnoredClients](const FGuid& ClientId){ return !IgnoredClients.Contains(ClientId); });
+	}
+
+	bool FPredictedStateObjectHierarchy::HasChildren(const FSoftObjectPath& Object) const
+	{
+		return Hierarchy.HasChildren(Object);
+	}
+}
diff --git a/Engine/Plugins/Developer/Concert/ConcertSync/ConcertSyncServer/Source/ConcertSyncServer/Private/Replication/Muting/PredictedStateObjectHierarchy.h b/Engine/Plugins/Developer/Concert/ConcertSync/ConcertSyncServer/Source/ConcertSyncServer/Private/Replication/Muting/PredictedStateObjectHierarchy.h
new file mode 100644
--- /dev/null
+++ b/Engine/Plugins/Developer/Concert/ConcertSync/ConcertSyncServer/Source/ConcertSyncServer/Private/Replication/Muting/PredictedStateObjectHierarchy.h
@@ -0,0 +1,59 @@
+﻿// Copyright Epic Games, Inc. All Rights Reserved.
+
+#pragma once
+
+#include "IMuteValidationObjectHierarchy.h"
+#include "Misc/ObjectPathHierarchy.h"
+#include "Replication/ConcertReplicationClient.h"
+
+#include <type_traits>
+
+struct FConcertReplicationStream;
+
+namespace UE::ConcertSyncServer::Replication
+{
+	class IRegistrationEnumerator;
+
+	/**
+	 * Implementation of IMuteValidationObjectHierarchy which allows you to add arbitrary client state.
+	 * Primarily used to validate a hierarchy that clients WILL have in the future.
+	 */
+	class FPredictedStateObjectHierarchy : public IMuteValidationObjectHierarchy
+	{
+	public:
+
+		/** Adds clients from a map binding client id to stream content. */
+		void AddClients(const TMap<FGuid, TArray<FConcertReplicationStream>>& Clients);
+		/** Adds only those entries that pass ShouldIncludeFilter returns true on. */
+		template<typename TFilterLambda> requires std::is_invocable_r_v<bool, TFilterLambda, const FGuid&>
+		void AddClients(const TMap<FGuid, TUniquePtr<FConcertReplicationClient>>& Clients, TFilterLambda&& ShouldIncludeFilter);
+		
+		/** Adds the streams a client will have. */
+		void AddClientData(const FGuid& ClientId, TConstArrayView<FConcertReplicationStream> Streams);
+		
+		//~ Begin IMuteValidationObjectHierarchy Interface
+		virtual bool IsObjectReferencedDirectly(const FSoftObjectPath& ObjectPath, TConstArrayView<FGuid> IgnoredClients) const override;
+		virtual bool HasChildren(const FSoftObjectPath& Object) const override;
+		//~ End IMuteValidationObjectHierarchy Interface
+
+	private:
+
+		/** The hierarchy built during construction. */
+		ConcertSyncCore::FObjectPathHierarchy Hierarchy;
+
+		/** Maps objects to the clients referencing them. */
+		TMap<FSoftObjectPath, TArray<FGuid>> ObjectReferencingClients;
+	};
+
+	template <typename TFilterLambda> requires std::is_invocable_r_v<bool, TFilterLambda, const FGuid&>
+	void FPredictedStateObjectHierarchy::AddClients(const TMap<FGuid, TUniquePtr<FConcertReplicationClient>>& Clients, TFilterLambda&& ShouldIncludeFilter)
+	{
+		for (const TPair<FGuid, TUniquePtr<FConcertReplicationClient>>& Pair : Clients)
+		{
+			if (ShouldIncludeFilter(Pair.Key))
+			{
+				AddClientData(Pair.Key, Pair.Value->GetStreamDescriptions());
+			}
+		}
+	}
+}
