Add FConcertReplication_PutState_Request which allows editing multiple clients remotely.
This feature is opt-in and enabled by default for Multi User.

diff --git a/Engine/Plugins/Developer/Concert/ConcertSync/ConcertSyncClient/Source/ConcertSyncClient/Private/Replication/Manager/ReplicationManager.cpp b/Engine/Plugins/Developer/Concert/ConcertSync/ConcertSyncClient/Source/ConcertSyncClient/Private/Replication/Manager/ReplicationManager.cpp
--- a/Engine/Plugins/Developer/Concert/ConcertSync/ConcertSyncClient/Source/ConcertSyncClient/Private/Replication/Manager/ReplicationManager.cpp
+++ b/Engine/Plugins/Developer/Concert/ConcertSync/ConcertSyncClient/Source/ConcertSyncClient/Private/Replication/Manager/ReplicationManager.cpp
@@ -4,7 +4,7 @@
 
 #include "IConcertSession.h"
 #include "ReplicationManagerState_Disconnected.h"
-#include "ReplicationManagerUtils.h"
+#include "Replication/Manager/Utils/ReplicationManagerUtils.h"
 
 #include "Replication/Messages/Handshake.h"
 
@@ -147,6 +147,13 @@ namespace UE::ConcertSyncClient::Replication
 			: MakeFulfilledPromise<FConcertReplication_RestoreContent_Response>().GetFuture(); 
 	}
 
+	TFuture<FConcertReplication_PutState_Response> FReplicationManager::PutClientState(FConcertReplication_PutState_Request Request)
+	{
+		return ensureMsgf(CurrentState, TEXT("StartAcceptingJoinRequests should have been called at this point."))
+			? CurrentState->PutClientState(MoveTemp(Request))
+			: MakeFulfilledPromise<FConcertReplication_PutState_Response>().GetFuture(); 
+	}
+
 	IConcertClientReplicationManager::FOnPreStreamsChanged& FReplicationManager::OnPreStreamsChanged()
 	{
 		// Check() to avoid returning some dummy static variable
diff --git a/Engine/Plugins/Developer/Concert/ConcertSync/ConcertSyncClient/Source/ConcertSyncClient/Private/Replication/Manager/ReplicationManager.h b/Engine/Plugins/Developer/Concert/ConcertSync/ConcertSyncClient/Source/ConcertSyncClient/Private/Replication/Manager/ReplicationManager.h
--- a/Engine/Plugins/Developer/Concert/ConcertSync/ConcertSyncClient/Source/ConcertSyncClient/Private/Replication/Manager/ReplicationManager.h
+++ b/Engine/Plugins/Developer/Concert/ConcertSync/ConcertSyncClient/Source/ConcertSyncClient/Private/Replication/Manager/ReplicationManager.h
@@ -46,6 +46,7 @@ namespace UE::ConcertSyncClient::Replication
 		virtual TFuture<FConcertReplication_ChangeMuteState_Response> ChangeMuteState(FConcertReplication_ChangeMuteState_Request Request) override;
 		virtual TFuture<FConcertReplication_QueryMuteState_Response> QueryMuteState(FConcertReplication_QueryMuteState_Request Request) override;
 		virtual TFuture<FConcertReplication_RestoreContent_Response> RestoreContent(FConcertReplication_RestoreContent_Request Request) override;
+		virtual TFuture<FConcertReplication_PutState_Response> PutClientState(FConcertReplication_PutState_Request Request) override;
 		virtual FOnPreStreamsChanged& OnPreStreamsChanged() override;
 		virtual FOnPostStreamsChanged& OnPostStreamsChanged() override;
 		virtual FOnPreAuthorityChanged& OnPreAuthorityChanged() override;
diff --git a/Engine/Plugins/Developer/Concert/ConcertSync/ConcertSyncClient/Source/ConcertSyncClient/Private/Replication/Manager/ReplicationManagerState.h b/Engine/Plugins/Developer/Concert/ConcertSync/ConcertSyncClient/Source/ConcertSyncClient/Private/Replication/Manager/ReplicationManagerState.h
--- a/Engine/Plugins/Developer/Concert/ConcertSync/ConcertSyncClient/Source/ConcertSyncClient/Private/Replication/Manager/ReplicationManagerState.h
+++ b/Engine/Plugins/Developer/Concert/ConcertSync/ConcertSyncClient/Source/ConcertSyncClient/Private/Replication/Manager/ReplicationManagerState.h
@@ -2,8 +2,8 @@
 
 #pragma once
 
-#include "ReplicationManagerUtils.h"
 #include "Replication/IConcertClientReplicationManager.h"
+#include "Replication/Manager/Utils/ReplicationManagerUtils.h"
 #include "Replication/Messages/Muting.h"
 
 namespace UE::ConcertSyncClient::Replication
@@ -35,6 +35,7 @@ namespace UE::ConcertSyncClient::Replication
 		virtual TFuture<FConcertReplication_ChangeMuteState_Response> ChangeMuteState(FConcertReplication_ChangeMuteState_Request) override { return MakeFulfilledPromise<FConcertReplication_ChangeMuteState_Response>(FConcertReplication_ChangeMuteState_Response{ EConcertReplicationMuteErrorCode::Rejected }).GetFuture(); };
 		virtual TFuture<FConcertReplication_QueryMuteState_Response> QueryMuteState(FConcertReplication_QueryMuteState_Request Request) override { return MakeFulfilledPromise<FConcertReplication_QueryMuteState_Response>().GetFuture(); }
 		virtual TFuture<FConcertReplication_RestoreContent_Response> RestoreContent(FConcertReplication_RestoreContent_Request Request) override { return MakeFulfilledPromise<FConcertReplication_RestoreContent_Response>().GetFuture(); }
+		virtual TFuture<FConcertReplication_PutState_Response> PutClientState(FConcertReplication_PutState_Request Request) override { return MakeFulfilledPromise<FConcertReplication_PutState_Response>().GetFuture(); }
 		virtual FOnPreStreamsChanged& OnPreStreamsChanged() override { return OnPreStreamsChangedDelegate; } 
 		virtual FOnPostStreamsChanged& OnPostStreamsChanged() override { return OnPostStreamsChangedDelegate; }
 		virtual FOnPreAuthorityChanged& OnPreAuthorityChanged() override { return OnPreAuthorityChangedDelegate; }
diff --git a/Engine/Plugins/Developer/Concert/ConcertSync/ConcertSyncClient/Source/ConcertSyncClient/Private/Replication/Manager/ReplicationManagerState_Connected.cpp b/Engine/Plugins/Developer/Concert/ConcertSync/ConcertSyncClient/Source/ConcertSyncClient/Private/Replication/Manager/ReplicationManagerState_Connected.cpp
--- a/Engine/Plugins/Developer/Concert/ConcertSync/ConcertSyncClient/Source/ConcertSyncClient/Private/Replication/Manager/ReplicationManagerState_Connected.cpp
+++ b/Engine/Plugins/Developer/Concert/ConcertSync/ConcertSyncClient/Source/ConcertSyncClient/Private/Replication/Manager/ReplicationManagerState_Connected.cpp
@@ -8,81 +8,22 @@
 #include "ReplicationManagerState_Disconnected.h"
 #include "Replication/ChangeStreamSharedUtils.h"
 #include "Replication/Formats/FullObjectFormat.h"
+#include "Replication/Messages/ChangeClientEvent.h"
 #include "Replication/Messages/Handshake.h"
 #include "Replication/Processing/ClientReplicationDataCollector.h"
 #include "Replication/Processing/ObjectReplicationApplierProcessor.h"
 #include "Replication/Processing/ObjectReplicationReceiver.h"
+#include "Utils/LocalSyncControl.h"
+#include "Utils/NetworkMessageLogging.h"
+#include "Utils/ReplicationManagerUtils.h"
 
 #include "Algo/RemoveIf.h"
-#include "JsonObjectConverter.h"
 #include "Misc/ScopeExit.h"
 #include "UObject/Package.h"
 #include "UObject/UObjectGlobals.h"
 
 namespace UE::ConcertSyncClient::Replication
 {
-	TAutoConsoleVariable<bool> CVarSimulateAuthorityTimeouts(
-		TEXT("Concert.Replication.SimulateAuthorityTimeouts"),
-		false,
-		TEXT("Whether the client should pretend that authority requests timed out instead of sending to the server.")
-		);
-	TAutoConsoleVariable<bool> CVarSimulateQueryTimeouts(
-		TEXT("Concert.Replication.SimulateQueryTimeouts"),
-		false,
-		TEXT("Whether the client should pretend that query requests timed out instead of sending to the server.")
-		);
-	TAutoConsoleVariable<bool> CVarSimulateStreamChangeTimeouts(
-		TEXT("Concert.Replication.SimulateStreamChangeTimeouts"),
-		false,
-		TEXT("Whether the client should pretend that stream change requests timed out instead of sending to the server.")
-		);
-
-	TAutoConsoleVariable<bool> CVarSimulateAuthorityRejection(
-		TEXT("Concert.Replication.SimulateAuthorityRejection"),
-		false,
-		TEXT("Whether the client should pretend that authority change requests were rejected.")
-		);
-	TAutoConsoleVariable<bool> CVarSimulateMuteRequestRejection(
-		TEXT("Concert.Replication.SimulateMuteRejection"),
-		false,
-		TEXT("Whether the client should pretend that mute change requests were rejected.")
-		);
-
-	TAutoConsoleVariable<bool> CVarLogStreamRequestsAndResponsesOnClient(
-		TEXT("Concert.Replication.LogStreamRequestsAndResponsesOnClient"),
-		false,
-		TEXT("Whether to log changes to streams.")
-		);
-	TAutoConsoleVariable<bool> CVarLogAuthorityRequestsAndResponsesOnClient(
-		TEXT("Concert.Replication.LogAuthorityRequestsAndResponsesOnClient"),
-		false,
-		TEXT("Whether to log changes to authority.")
-		);
-	TAutoConsoleVariable<bool> CVarLogMuteRequestsAndResponsesOnClient(
-		TEXT("Concert.Replication.LogMuteRequestsAndResponsesOnClient"),
-		false,
-		TEXT("Whether to log changes to the mute state.")
-		);
-	TAutoConsoleVariable<bool> CVarLogRestoreContentRequestsAndResponsesOnClient(
-		TEXT("Concert.Replication.LogRestoreContentRequestsAndResponsesOnClient"),
-		false,
-		TEXT("Whether to log restore content requests and responses.")
-		);
-
-	namespace Private
-	{
-		template<typename TMessage>
-		static void LogNetworkMessage(const TAutoConsoleVariable<bool>& ShouldLog, const TMessage& Message)
-		{
-			if (ShouldLog.GetValueOnAnyThread())
-			{
-				FString JsonString;
-				FJsonObjectConverter::UStructToJsonObjectString(TMessage::StaticStruct(), &Message, JsonString, 0, 0);
-				UE_LOG(LogConcert, Log, TEXT("%s\n%s"), *TMessage::StaticStruct()->GetName(), *JsonString);
-			}
-		}
-	}
-	
 	FReplicationManagerState_Connected::FReplicationManagerState_Connected(
 		TSharedRef<IConcertClientSession> InLiveSession,
 		IConcertClientReplicationBridge& ReplicationBridge,
@@ -127,6 +68,7 @@ namespace UE::ConcertSyncClient::Replication
 	{
 		// Technically not needed due to AddSP but let's be nice and clean up after ourselves
 		LiveSession->OnTick().RemoveAll(this);
+		LiveSession->UnregisterCustomEventHandler<FConcertReplication_ChangeClientEvent>(this);
 	}
 
 	TFuture<FJoinReplicatedSessionResult> FReplicationManagerState_Connected::JoinReplicationSession(FJoinReplicatedSessionArgs Args)
@@ -168,26 +110,25 @@ namespace UE::ConcertSyncClient::Replication
 		
 		// Stop replicating removed objects right now: the server will remove authority after processing this request.
 		// At that point, it will log errors for receiving replication data from a client without authority.
-		HandleReleasingReplicatedObjects(Args);
+		ApplyAuthorityChangeRemovedObjects(Args);
 		// We don't need worry about updating sync control until it is processed below - the local client will not attempt to replicate the object
 		// because we just locally updated the authority cache.
 
-		Private::LogNetworkMessage(CVarLogAuthorityRequestsAndResponsesOnClient, Args);
+		LogNetworkMessage(CVarLogAuthorityRequestsAndResponsesOnClient, Args);
 		return LiveSession->SendCustomRequest<FConcertReplication_ChangeAuthority_Request, FConcertReplication_ChangeAuthority_Response>(Args, LiveSession->GetSessionServerEndpointId())
 			.Next([WeakThis = TWeakPtr<FReplicationManagerState_Connected>(SharedThis(this)), Args](FConcertReplication_ChangeAuthority_Response&& Response) mutable
 			{
-				Private::LogNetworkMessage(CVarLogAuthorityRequestsAndResponsesOnClient, Response);
+				LogNetworkMessage(CVarLogAuthorityRequestsAndResponsesOnClient, Response);
 				
 				if (const TSharedPtr<FReplicationManagerState_Connected> ThisPin = WeakThis.Pin()
 					; ThisPin && Response.ErrorCode == EReplicationResponseErrorCode::Handled)
 				{
-					ThisPin->SyncControl.ProcessAuthorityChange(Args, Response);
-					ThisPin->UpdateReplicatedObjectsAfterAuthorityChange(MoveTemp(Args), Response);
+					ThisPin->FinalizePredictedAuthorityChange(Args, Response.RejectedObjects, Response.SyncControl);
 				}
 				else if (ThisPin && Response.ErrorCode == EReplicationResponseErrorCode::Timeout)
 				{
 					// HandleReleasingReplicatedObjects caused Args.ReleaseAuthority to stop being replicated. Revert.
-					ThisPin->RevertReleasingReplicatedObjects(Args);
+					ThisPin->RevertAuthorityChangeReleasedObjects(Args);
 				}
 
 				return FConcertReplication_ChangeAuthority_Response { MoveTemp(Response) };
@@ -223,26 +164,25 @@ namespace UE::ConcertSyncClient::Replication
 		
 		// Stop replicating removed objects right now: the server will remove authority after processing this request.
 		// At that point, it will log errors for receiving replication data from a client without authority.
-		TMap<FSoftObjectPath, TArray<FGuid>> PredictedChange = HandleRemovingReplicatedObjects(Args);
+		TMap<FSoftObjectPath, TArray<FGuid>> PredictedChange = PredictAndApplyStreamChangeRemovedObjects(Args);
 		// We don't need worry about updating sync control until it is processed below - the local client will not attempt to replicate the object
 		// because we just locally updated the replication cache.
 		
-		Private::LogNetworkMessage(CVarLogStreamRequestsAndResponsesOnClient, Args);
+		LogNetworkMessage(CVarLogStreamRequestsAndResponsesOnClient, Args);
 		return LiveSession->SendCustomRequest<FConcertReplication_ChangeStream_Request, FConcertReplication_ChangeStream_Response>(Args, LiveSession->GetSessionServerEndpointId())
 			.Next([WeakThis = TWeakPtr<FReplicationManagerState_Connected>(SharedThis(this)), Args, PredictedChange = MoveTemp(PredictedChange)](FConcertReplication_ChangeStream_Response&& Response)
 			{
-				Private::LogNetworkMessage(CVarLogStreamRequestsAndResponsesOnClient, Response);
+				LogNetworkMessage(CVarLogStreamRequestsAndResponsesOnClient, Response);
 
 				const TSharedPtr<FReplicationManagerState_Connected> ThisPin = WeakThis.Pin();
 				if (ThisPin && Response.IsSuccess())
 				{
-					ThisPin->SyncControl.ProcessStreamChange(Args);
-					ThisPin->UpdateReplicatedObjectsAfterStreamChange(Args);
+					ThisPin->FinalizePredictedStreamChange(Args);
 				}
 				else if (ThisPin && Response.ErrorCode == EReplicationResponseErrorCode::Timeout)
 				{
 					// HandleRemovingReplicatedObjects caused Request.ObjectsToRemove to stop being replicated. Revert.
-					ThisPin->RevertRemovingReplicatedObjects(PredictedChange);
+					ThisPin->RevertPredictedStreamChangeRemovedObjects(PredictedChange);
 				}
 				
 				return FConcertReplication_ChangeStream_Response { MoveTemp(Response) };
@@ -292,11 +232,11 @@ namespace UE::ConcertSyncClient::Replication
 
 		FLocalSyncControl::FPredictedObjectRemoval PredictedChanges = SyncControl.PredictAndApplyMuteChanges(Request);
 		
-		Private::LogNetworkMessage(CVarLogMuteRequestsAndResponsesOnClient, Request);
+		LogNetworkMessage(CVarLogMuteRequestsAndResponsesOnClient, Request);
 		return LiveSession->SendCustomRequest<FConcertReplication_ChangeMuteState_Request, FConcertReplication_ChangeMuteState_Response>(Request, LiveSession->GetSessionServerEndpointId())
 			.Next([WeakThis = TWeakPtr<FReplicationManagerState_Connected>(SharedThis(this)), PredictedChanges = MoveTemp(PredictedChanges)](FConcertReplication_ChangeMuteState_Response&& Response)
 			{
-				Private::LogNetworkMessage(CVarLogMuteRequestsAndResponsesOnClient, Response);
+				LogNetworkMessage(CVarLogMuteRequestsAndResponsesOnClient, Response);
 				if (const TSharedPtr<FReplicationManagerState_Connected> This = WeakThis.Pin())
 				{
 					This->SyncControl.ApplyOrRevertMuteResponse(PredictedChanges, Response);
@@ -330,12 +270,12 @@ namespace UE::ConcertSyncClient::Replication
 		// We want the response to contain ClientInfo to update our internal state - so set the SendNewState flag
 		const bool bWantedNewState = EnumHasAnyFlags(Request.Flags, EConcertReplicationRestoreContentFlags::SendNewState);
 		Request.Flags |= EConcertReplicationRestoreContentFlags::SendNewState;
-		Private::LogNetworkMessage(CVarLogRestoreContentRequestsAndResponsesOnClient, Request);
+		LogNetworkMessage(CVarLogRestoreContentRequestsAndResponsesOnClient, Request);
 		
 		return LiveSession->SendCustomRequest<FConcertReplication_RestoreContent_Request, FConcertReplication_RestoreContent_Response>(Request, LiveSession->GetSessionServerEndpointId())
 			.Next([WeakThis = TWeakPtr<FReplicationManagerState_Connected>(SharedThis(this)), PredictedChanges = MoveTemp(PredictedChanges)](FConcertReplication_RestoreContent_Response&& Response)
 			{
-				Private::LogNetworkMessage(CVarLogRestoreContentRequestsAndResponsesOnClient, Response);
+				LogNetworkMessage(CVarLogRestoreContentRequestsAndResponsesOnClient, Response);
 				
 				const TSharedPtr<FReplicationManagerState_Connected> This = WeakThis.Pin();
 				if (!This)
@@ -366,6 +306,7 @@ namespace UE::ConcertSyncClient::Replication
 	void FReplicationManagerState_Connected::OnEnterState()
 	{
 		LiveSession->OnTick().AddSP(this, &FReplicationManagerState_Connected::Tick);
+		LiveSession->RegisterCustomEventHandler<FConcertReplication_ChangeClientEvent>(this, &FReplicationManagerState_Connected::HandleChangeClientEvent);
 	}
 
 	void FReplicationManagerState_Connected::Tick(IConcertClientSession& Session, float DeltaTime)
@@ -376,6 +317,35 @@ namespace UE::ConcertSyncClient::Replication
 		ReplicationApplier.ProcessObjects(Params);
 	}
 
+	TMap<FSoftObjectPath, TArray<FGuid>> FReplicationManagerState_Connected::PredictAndApplyStreamChangeRemovedObjects(const FConcertReplication_ChangeStream_Request& Request)
+	{
+		TMap<FSoftObjectPath, TArray<FGuid>> PredicatedRemovedObjects = ComputeRemovedObjects(RegisteredStreams, Request);
+		ApplyStreamChangeRemovedObjects(PredicatedRemovedObjects);
+		return PredicatedRemovedObjects;
+	}
+
+	void FReplicationManagerState_Connected::ApplyStreamChangeRemovedObjects(const TMap<FSoftObjectPath, TArray<FGuid>>& PredicatedRemovedObjects)
+	{
+		for (const TPair<FSoftObjectPath, TArray<FGuid>>& RemovedObjectInfo : PredicatedRemovedObjects)
+		{
+			ReplicationDataSource.RemoveReplicatedObjectStreams(RemovedObjectInfo.Key, RemovedObjectInfo.Value);
+		}
+	}
+
+	void FReplicationManagerState_Connected::RevertPredictedStreamChangeRemovedObjects(const TMap<FSoftObjectPath, TArray<FGuid>>& PredictedChange)
+	{
+		for (const TPair<FSoftObjectPath, TArray<FGuid>>& RemovedObjectInfo : PredictedChange)
+		{
+			ReplicationDataSource.AddReplicatedObjectStreams(RemovedObjectInfo.Key, RemovedObjectInfo.Value);
+		}
+	}
+	
+	void FReplicationManagerState_Connected::FinalizePredictedStreamChange(const FConcertReplication_ChangeStream_Request& StreamChange)
+	{
+		SyncControl.ProcessStreamChange(StreamChange);
+		UpdateReplicatedObjectsAfterStreamChange(StreamChange);
+	}
+	
 	void FReplicationManagerState_Connected::UpdateReplicatedObjectsAfterStreamChange(const FConcertReplication_ChangeStream_Request& Request)
 	{
 		OnPreStreamsChangedDelegate.Broadcast();
@@ -411,66 +381,42 @@ namespace UE::ConcertSyncClient::Replication
 		}
 	}
 
-	namespace Private
+	void FReplicationManagerState_Connected::ApplyAuthorityChangeRemovedObjects(const FConcertReplication_ChangeAuthority_Request& Request)
 	{
-		static void ForEachObjectRemovedFromStreams(const FConcertReplication_ChangeStream_Request& Request, TFunctionRef<void(const FSoftObjectPath& ObjectPath, const TArray<FGuid>& Streams)> Callback)
+		for (const TPair<FSoftObjectPath, FConcertStreamArray>& ReleaseAuthority : Request.ReleaseAuthority)
 		{
-			
+			ReplicationDataSource.RemoveReplicatedObjectStreams(ReleaseAuthority.Key, ReleaseAuthority.Value.StreamIds);
 		}
 	}
 
-	TMap<FSoftObjectPath, TArray<FGuid>> FReplicationManagerState_Connected::HandleRemovingReplicatedObjects(const FConcertReplication_ChangeStream_Request& Request)
+	void FReplicationManagerState_Connected::RevertAuthorityChangeReleasedObjects(const FConcertReplication_ChangeAuthority_Request& Request)
 	{
-		TMap<FSoftObjectPath, TArray<FGuid>> BundledRemovedObjects;
-		for (const FConcertObjectInStreamID& RemovedObject : Request.ObjectsToRemove)
-		{
-			BundledRemovedObjects.FindOrAdd(RemovedObject.Object).AddUnique(RemovedObject.StreamId);
-		}
-		for (const FConcertReplicationStream& Stream : RegisteredStreams)
-		{
-			const FGuid& StreamId = Stream.BaseDescription.Identifier;
-			if (!Request.StreamsToRemove.Contains(StreamId))
-			{
-				continue;
-			}
-
-			for (const TPair<FSoftObjectPath, FConcertReplicatedObjectInfo>& Pair : Stream.BaseDescription.ReplicationMap.ReplicatedObjects)
-			{
-				BundledRemovedObjects.FindOrAdd(Pair.Key).AddUnique(StreamId);
-			}
-		}
-		
-		for (const TPair<FSoftObjectPath, TArray<FGuid>>& RemovedObjectInfo : BundledRemovedObjects)
+		for (const TPair<FSoftObjectPath, FConcertStreamArray>& ReleaseAuthority : Request.ReleaseAuthority)
 		{
-			ReplicationDataSource.RemoveReplicatedObjectStreams(RemovedObjectInfo.Key, RemovedObjectInfo.Value);
+			ReplicationDataSource.AddReplicatedObjectStreams(ReleaseAuthority.Key, ReleaseAuthority.Value.StreamIds);
 		}
-		return BundledRemovedObjects;
 	}
 
-	void FReplicationManagerState_Connected::RevertRemovingReplicatedObjects(const TMap<FSoftObjectPath, TArray<FGuid>>& PredictedChange)
+	void FReplicationManagerState_Connected::FinalizePredictedAuthorityChange(const FConcertReplication_ChangeAuthority_Request& AuthorityChange, const TMap<FSoftObjectPath, FConcertStreamArray>& RejectedObjects, const FConcertReplication_ChangeSyncControl& SyncControlChange)
 	{
-		for (const TPair<FSoftObjectPath, TArray<FGuid>>& RemovedObjectInfo : PredictedChange)
-		{
-			ReplicationDataSource.AddReplicatedObjectStreams(RemovedObjectInfo.Key, RemovedObjectInfo.Value);
-		}
+		SyncControl.ProcessAuthorityChange(AuthorityChange, SyncControlChange);
+		UpdateReplicatedObjectsAfterAuthorityChange(AuthorityChange, RejectedObjects);
 	}
-
-	void FReplicationManagerState_Connected::UpdateReplicatedObjectsAfterAuthorityChange(FConcertReplication_ChangeAuthority_Request&& Request, const FConcertReplication_ChangeAuthority_Response& Response)
+	
+	void FReplicationManagerState_Connected::UpdateReplicatedObjectsAfterAuthorityChange(const FConcertReplication_ChangeAuthority_Request& Request, const TMap<FSoftObjectPath, FConcertStreamArray>& RejectedObjects)
 	{
 		OnPreAuthorityChangedDelegate.Broadcast();
 		ON_SCOPE_EXIT{ OnPostAuthorityChangedDelegate.Broadcast(); };
 		
-		for (TPair<FSoftObjectPath, FConcertStreamArray>& TakeAuthority : Request.TakeAuthority)
+		for (const TPair<FSoftObjectPath, FConcertStreamArray>& TakeAuthority : Request.TakeAuthority)
 		{
 			const FSoftObjectPath& ReplicatedObject = TakeAuthority.Key;
-			// Request will be discarded so ...
-			FConcertStreamArray& ReplicatedStreams = TakeAuthority.Value;
+			FConcertStreamArray ReplicatedStreams = TakeAuthority.Value;
 			
 			UE_CLOG(ReplicatedStreams.StreamIds.IsEmpty(), LogConcert, Warning, TEXT("Your FConcertReplication_ChangeAuthority_Request::TakeAuthority request contained empty stream ID array for object %s"), *ReplicatedObject.ToString());
-			const FConcertStreamArray* RejectedStreams = Response.RejectedObjects.Find(ReplicatedObject);
+			const FConcertStreamArray* RejectedStreams = RejectedObjects.Find(ReplicatedObject);
 			if (RejectedStreams)
 			{
-				// ... reuse the memory
 				ReplicatedStreams.StreamIds.SetNum(Algo::RemoveIf(ReplicatedStreams.StreamIds, [RejectedStreams](const FGuid& Stream)
 				{
 					return RejectedStreams->StreamIds.Contains(Stream);
@@ -485,22 +431,6 @@ namespace UE::ConcertSyncClient::Replication
 		}
 	}
 
-	void FReplicationManagerState_Connected::HandleReleasingReplicatedObjects(const FConcertReplication_ChangeAuthority_Request& Request)
-	{
-		for (const TPair<FSoftObjectPath, FConcertStreamArray>& ReleaseAuthority : Request.ReleaseAuthority)
-		{
-			ReplicationDataSource.RemoveReplicatedObjectStreams(ReleaseAuthority.Key, ReleaseAuthority.Value.StreamIds);
-		}
-	}
-
-	void FReplicationManagerState_Connected::RevertReleasingReplicatedObjects(const FConcertReplication_ChangeAuthority_Request& Request)
-	{
-		for (const TPair<FSoftObjectPath, FConcertStreamArray>& ReleaseAuthority : Request.ReleaseAuthority)
-		{
-			ReplicationDataSource.AddReplicatedObjectStreams(ReleaseAuthority.Key, ReleaseAuthority.Value.StreamIds);
-		}
-	}
-
 	void FReplicationManagerState_Connected::UpdateReplicatedObjectAfterServerSideChange(const FConcertQueriedClientInfo& NewState)
 	{
 		OnPreStreamsChangedDelegate.Broadcast();
diff --git a/Engine/Plugins/Developer/Concert/ConcertSync/ConcertSyncClient/Source/ConcertSyncClient/Private/Replication/Manager/ReplicationManagerState_Connected.h b/Engine/Plugins/Developer/Concert/ConcertSync/ConcertSyncClient/Source/ConcertSyncClient/Private/Replication/Manager/ReplicationManagerState_Connected.h
--- a/Engine/Plugins/Developer/Concert/ConcertSync/ConcertSyncClient/Source/ConcertSyncClient/Private/Replication/Manager/ReplicationManagerState_Connected.h
+++ b/Engine/Plugins/Developer/Concert/ConcertSync/ConcertSyncClient/Source/ConcertSyncClient/Private/Replication/Manager/ReplicationManagerState_Connected.h
@@ -5,13 +5,14 @@
 #include "ConcertSyncSessionFlags.h"
 #include "ReplicationManagerState.h"
 #include "Replication/Formats/IObjectReplicationFormat.h"
-#include "Replication/Misc/LocalSyncControl.h"
+#include "Replication/Messages/ChangeClientEvent.h"
 #include "Replication/Processing/ClientReplicationDataCollector.h"
 #include "Replication/Processing/ClientReplicationDataQueuer.h"
 #include "Replication/Processing/ObjectReplicationApplierProcessor.h"
 #include "Replication/Processing/ObjectReplicationReceiver.h"
 #include "Replication/Processing/ObjectReplicationSender.h"
 #include "Replication/Processing/Proxy/ObjectProcessorProxy_Frequency.h"
+#include "Utils/LocalSyncControl.h"
 
 class IConcertClientSession;
 
@@ -29,7 +30,7 @@ namespace UE::ConcertSyncClient::Replication
 	class FReplicationManagerState_Connected : public FReplicationManagerState
 	{
 	public:
-
+		
 		FReplicationManagerState_Connected(
 			TSharedRef<IConcertClientSession> InLiveSession,
 			IConcertClientReplicationBridge& ReplicationBridge UE_LIFETIMEBOUND,
@@ -57,6 +58,7 @@ namespace UE::ConcertSyncClient::Replication
 		virtual TFuture<FConcertReplication_ChangeMuteState_Response> ChangeMuteState(FConcertReplication_ChangeMuteState_Request) override;
 		virtual TFuture<FConcertReplication_QueryMuteState_Response> QueryMuteState(FConcertReplication_QueryMuteState_Request Request) override;
 		virtual TFuture<FConcertReplication_RestoreContent_Response> RestoreContent(FConcertReplication_RestoreContent_Request Request) override;
+		virtual TFuture<FConcertReplication_PutState_Response> PutClientState(FConcertReplication_PutState_Request Request) override;
 		//~ End IConcertClientReplicationManager Interface
 
 	private:
@@ -109,21 +111,28 @@ namespace UE::ConcertSyncClient::Replication
 		 * It is configured in the project settings TODO: Add config
 		 */
 		void Tick(IConcertClientSession& Session, float DeltaTime);
-		
+
+		/** Handle the server telling us that our state has changed. */
+		void HandleChangeClientEvent(const FConcertSessionContext& Context, const FConcertReplication_ChangeClientEvent& Event);
+
+		/** Changes the local state assuming that Request will succeed. */
+		TMap<FSoftObjectPath, TArray<FGuid>> PredictAndApplyStreamChangeRemovedObjects(const FConcertReplication_ChangeStream_Request& Request);
+		void ApplyStreamChangeRemovedObjects(const TMap<FSoftObjectPath, TArray<FGuid>>& PredicatedRemovedObjects);
+		/** Reverts changes previously made by PredictStreamChangeRemovedObjects. */
+		void RevertPredictedStreamChangeRemovedObjects(const TMap<FSoftObjectPath, TArray<FGuid>>& PredictedChange);
+		/** Applies stream changes that we previously predicted using PredictStreamChangeRemovedObjects. */
+		void FinalizePredictedStreamChange(const FConcertReplication_ChangeStream_Request& StreamChange);
 		/** Updates replicated objects affected by the change request. */
 		void UpdateReplicatedObjectsAfterStreamChange(const FConcertReplication_ChangeStream_Request& Request);
-		TMap<FSoftObjectPath, TArray<FGuid>> HandleRemovingReplicatedObjects(const FConcertReplication_ChangeStream_Request& Request);
-		void RevertRemovingReplicatedObjects(const TMap<FSoftObjectPath, TArray<FGuid>>& PredictedChange);
-
-		/**
-		 * Updates the objects which should be replicated after changing authority.
-		 * 
-		 * @note Request is accepted as && because this function rewrites its memory when looking at rejections.
-		 * Since the request was already sent to the server it can just contain trash after.
-		 */
-		void UpdateReplicatedObjectsAfterAuthorityChange(FConcertReplication_ChangeAuthority_Request&& Request, const FConcertReplication_ChangeAuthority_Response& Response);
-		void HandleReleasingReplicatedObjects(const FConcertReplication_ChangeAuthority_Request& Request);
-		void RevertReleasingReplicatedObjects(const FConcertReplication_ChangeAuthority_Request& Request);
+		
+		/** Changes the local state assuming that Request will succeed. */
+		void ApplyAuthorityChangeRemovedObjects(const FConcertReplication_ChangeAuthority_Request& Request);
+		/** Reverts changes previously made by PredictAuthorityChangeReleasedObjects. */
+		void RevertAuthorityChangeReleasedObjects(const FConcertReplication_ChangeAuthority_Request& Request);
+		/** Applies authority changes that we previously predicted using PredictAuthorityChangeReleasedObjects. */
+		void FinalizePredictedAuthorityChange(const FConcertReplication_ChangeAuthority_Request& AuthorityChange, const TMap<FSoftObjectPath, FConcertStreamArray>& RejectedObjects, const FConcertReplication_ChangeSyncControl& SyncControlChange);
+		/** Updates the objects which should be replicated after changing authority. */
+		void UpdateReplicatedObjectsAfterAuthorityChange(const FConcertReplication_ChangeAuthority_Request& Request, const TMap<FSoftObjectPath, FConcertStreamArray>& RejectedObjects);
 
 		/** Updates the objects which should be replicated after they have been reset to a completely new state (e.g. when restoring session content manually). */
 		void UpdateReplicatedObjectAfterServerSideChange(const FConcertQueriedClientInfo& NewState);
diff --git a/Engine/Plugins/Developer/Concert/ConcertSync/ConcertSyncClient/Source/ConcertSyncClient/Private/Replication/Manager/ReplicationManagerState_Connected_PutState.cpp b/Engine/Plugins/Developer/Concert/ConcertSync/ConcertSyncClient/Source/ConcertSyncClient/Private/Replication/Manager/ReplicationManagerState_Connected_PutState.cpp
new file mode 100644
--- /dev/null
+++ b/Engine/Plugins/Developer/Concert/ConcertSync/ConcertSyncClient/Source/ConcertSyncClient/Private/Replication/Manager/ReplicationManagerState_Connected_PutState.cpp
@@ -0,0 +1,235 @@
+﻿// Copyright Epic Games, Inc. All Rights Reserved.
+
+#include "ReplicationManagerState_Connected.h"
+
+#include "ConcertLogGlobal.h"
+#include "IConcertSession.h"
+#include "ReplicationManagerState_Disconnected.h"
+#include "Replication/ChangeStreamSharedUtils.h"
+#include "Replication/Messages/ChangeClientEvent.h"
+#include "Replication/Processing/ClientReplicationDataCollector.h"
+#include "Utils/LocalSyncControl.h"
+#include "Utils/NetworkMessageLogging.h"
+#include "Utils/ReplicationManagerUtils.h"
+
+#include "Misc/ScopeExit.h"
+#include "UObject/Package.h"
+#include "UObject/UObjectGlobals.h"
+
+namespace UE::ConcertSyncClient::Replication
+{
+	namespace Private
+	{
+		static TOptional<FConcertReplication_ChangeStream_Request> BuildDeltaStreamChange(
+			const FGuid& EndpointId,
+			const TConstArrayView<FConcertReplicationStream>& Streams,
+			const FConcertReplication_PutState_Request& Request
+			)
+		{
+			const FConcertReplicationStreamArray* LocallyNewStreams = Request.NewStreams.Find(EndpointId);
+			if (!LocallyNewStreams)
+			{
+				return {};
+			}
+			
+			FConcertReplication_ChangeStream_Request StreamRequest;
+			Algo::Transform(Streams, StreamRequest.StreamsToRemove, [](const FConcertReplicationStream& Stream){ return Stream.BaseDescription.Identifier; });
+			StreamRequest.StreamsToAdd = LocallyNewStreams->Streams;
+			return StreamRequest;
+		}
+
+		static TOptional<FConcertReplication_ChangeAuthority_Request> BuildDeltaAuthorityChange(const FGuid& EndpointId, const FConcertReplication_PutState_Request& Request)
+		{
+			const FConcertObjectInStreamArray* LocallyNewAuthority = Request.NewAuthorityState.Find(EndpointId);
+			if (!LocallyNewAuthority)
+			{
+				return {};
+			}
+			
+			FConcertReplication_ChangeAuthority_Request AuthorityRequest;
+			for (const FConcertObjectInStreamID& ObjectId : LocallyNewAuthority->Objects)
+			{
+				AuthorityRequest.TakeAuthority.FindOrAdd(ObjectId.Object).StreamIds.AddUnique(ObjectId.StreamId);
+			}
+			return AuthorityRequest;
+		}
+
+		static FConcertReplication_ChangeSyncControl BuildRetainedSyncControl(
+			const FGuid& EndpointId,
+			const FLocalSyncControl& SyncControl,
+			const FConcertReplication_PutState_Request& Request
+			)
+		{
+			FConcertReplication_ChangeSyncControl RetainedSyncControl;
+			
+			const FConcertObjectInStreamArray* LocallyNewAuthority = Request.NewAuthorityState.Find(EndpointId);
+			if (!LocallyNewAuthority)
+			{
+				// The client has not specified authority changes, so the sync control it has right now should be retained
+				SyncControl.EnumerateAllowedObjects([&RetainedSyncControl](const FConcertObjectInStreamID& ObjectId)
+				{
+					RetainedSyncControl.NewControlStates.Add(ObjectId, true);
+					return EBreakBehavior::Continue;
+				});
+				return RetainedSyncControl;
+			}
+			
+			for (const FConcertObjectInStreamID& ObjectId : LocallyNewAuthority->Objects)
+			{
+				if (SyncControl.IsObjectAllowed(ObjectId))
+				{
+					RetainedSyncControl.NewControlStates.Add(ObjectId, true);
+				}
+			}
+			return RetainedSyncControl;
+		}
+
+		static void RemoveEntriesThatShouldNotBeRetained(
+			const FConcertReplication_ChangeSyncControl& ResponseSyncControl,
+			const TConstArrayView<FConcertReplicationStream>& RegisteredStreams,
+			FConcertReplication_ChangeSyncControl& OutRetainedSyncControl
+			)
+		{
+			for (auto It = OutRetainedSyncControl.NewControlStates.CreateIterator(); It; ++It)
+			{
+				const FConcertObjectInStreamID& ObjectId = It->Key;
+
+				const bool bIsOverridenByResponse = ResponseSyncControl.NewControlStates.Contains(ObjectId);
+				const bool bIsRegistered = RegisteredStreams.ContainsByPredicate([&ObjectId](const FConcertReplicationStream& Stream)
+				{
+					const bool bIsSameStream = Stream.BaseDescription.Identifier == ObjectId.StreamId;
+					const bool bIsContainedInStream = Stream.BaseDescription.ReplicationMap.ReplicatedObjects.Contains(ObjectId.Object);
+					return bIsSameStream && bIsContainedInStream;
+				});
+
+				// If the object is contained in the response, the sync control has changed so we should not reapply it.
+				// If the request removed the object, it should also not be applied.
+				if (bIsOverridenByResponse || !bIsRegistered)
+				{
+					It.RemoveCurrent();
+				}
+			}
+		}
+
+		struct FPutStateChange
+		{
+			const TOptional<FConcertReplication_ChangeStream_Request> StreamChange;
+			const TOptional<FConcertReplication_ChangeAuthority_Request> AuthorityChange;
+			const TMap<FSoftObjectPath, TArray<FGuid>> PredictedStreamChange;
+			const FLocalSyncControl::FPredictedObjectRemoval PredictedMuteChange;
+			FConcertReplication_ChangeSyncControl RetainedSyncControl;
+
+			FPutStateChange(
+				const TConstArrayView<FConcertReplicationStream>& RegisteredStreams,
+				FLocalSyncControl::FPredictedObjectRemoval PredictedMuteChange,
+				const FLocalSyncControl& SyncControl,
+				const FGuid& ClientId,
+				const FConcertReplication_PutState_Request& Request
+				)
+				: StreamChange(BuildDeltaStreamChange(ClientId, RegisteredStreams, Request))
+				, AuthorityChange(BuildDeltaAuthorityChange(ClientId, Request))
+				, PredictedStreamChange(StreamChange ? ComputeRemovedObjects(RegisteredStreams, *StreamChange) : TMap<FSoftObjectPath, TArray<FGuid>>{})
+				, PredictedMuteChange(MoveTemp(PredictedMuteChange))
+				, RetainedSyncControl(BuildRetainedSyncControl(ClientId, SyncControl, Request))
+			{}
+		};
+	}
+
+	TFuture<FConcertReplication_PutState_Response> FReplicationManagerState_Connected::PutClientState(FConcertReplication_PutState_Request Request)
+	{
+		using namespace Private;
+		if (!EnumHasAnyFlags(SessionFlags, EConcertSyncSessionFlags::ShouldEnableRemoteEditing))
+		{
+			return MakeFulfilledPromise<FConcertReplication_PutState_Response>(FConcertReplication_PutState_Response{ EConcertReplicationPutStateResponseCode::FeatureDisabled }).GetFuture();
+		}
+
+		// If the request changes the local state, we follow the approaches from like ChangeStreams, ChangeAuthority, and ChangeMuteState by predicatively
+		// updating our local state if changing the local client as well... [1]
+		const FGuid& ClientId = LiveSession->GetSessionClientEndpointId();
+		FLocalSyncControl::FPredictedObjectRemoval PredictedMuteChange = SyncControl.PredictAndApplyMuteChanges(Request.MuteChange);
+		FPutStateChange Change(RegisteredStreams, MoveTemp(PredictedMuteChange), SyncControl, ClientId, Request);
+		if (Change.StreamChange)
+		{
+			ApplyStreamChangeRemovedObjects(Change.PredictedStreamChange);
+		}
+		if (Change.AuthorityChange)
+		{
+			ApplyAuthorityChangeRemovedObjects(*Change.AuthorityChange);
+		}
+		
+		LogNetworkMessage(CVarLogChangeClientsRequestsAndResponsesOnClient, Request);
+		return LiveSession->SendCustomRequest<FConcertReplication_PutState_Request, FConcertReplication_PutState_Response>(Request, LiveSession->GetSessionServerEndpointId())
+			.Next([WeakThis = TWeakPtr<FReplicationManagerState_Connected>(SharedThis(this)), Change = MoveTemp(Change)](FConcertReplication_PutState_Response&& Response) mutable
+			{
+				LogNetworkMessage(CVarLogChangeClientsRequestsAndResponsesOnClient, Response);
+				const TSharedPtr<FReplicationManagerState_Connected> ThisPin = WeakThis.Pin();
+				if (!ThisPin)
+				{
+					return Response;
+				}
+
+				const bool bIsSuccess = Response.IsSuccess();
+				if (bIsSuccess)
+				{
+					// We treat PutState as a separate stream operation followed by an authority operation.
+					// Applying the stream bit removes sync control... [2]
+					if (Change.StreamChange)
+					{
+						ThisPin->FinalizePredictedStreamChange(*Change.StreamChange);
+					}
+
+					// [2]... but we need to add back sync control for those objects that we retained sync control over.
+					// The server tells us which objects have changed sync control; those unchanged we must retain, so add them back now.
+					RemoveEntriesThatShouldNotBeRetained(Response.SyncControl, ThisPin->RegisteredStreams, Change.RetainedSyncControl);
+					ThisPin->SyncControl.ProcessSyncControlChange(Change.RetainedSyncControl);
+					
+					if (Change.AuthorityChange)
+					{
+						ThisPin->FinalizePredictedAuthorityChange(*Change.AuthorityChange, {}, Response.SyncControl);
+					}
+				}
+				else
+				{
+					// [1]... and we may have to revert predictive changes
+					if (Change.StreamChange)
+					{
+						// HandleRemovingReplicatedObjects caused Request.ObjectsToRemove to stop being replicated. Revert.
+						ThisPin->RevertPredictedStreamChangeRemovedObjects(Change.PredictedStreamChange);
+					}
+					if (Change.AuthorityChange)
+					{
+						// HandleReleasingReplicatedObjects caused Args.ReleaseAuthority to stop being replicated. Revert.
+						ThisPin->RevertAuthorityChangeReleasedObjects(*Change.AuthorityChange);
+					}
+					
+					ThisPin->SyncControl.ApplyOrRevertMuteResponse(Change.PredictedMuteChange, { .ErrorCode = EConcertReplicationMuteErrorCode::Rejected });
+				}
+				
+				const EConcertReplicationMuteErrorCode DummyMuteResponseCode = bIsSuccess ? EConcertReplicationMuteErrorCode::Accepted : EConcertReplicationMuteErrorCode::Rejected;
+				ThisPin->SyncControl.ApplyOrRevertMuteResponse(Change.PredictedMuteChange, { .ErrorCode = DummyMuteResponseCode });
+				
+				return Response;
+			});
+	}
+
+	void FReplicationManagerState_Connected::HandleChangeClientEvent(const FConcertSessionContext& Context, const FConcertReplication_ChangeClientEvent& Event)
+	{
+		LogNetworkMessage(CVarLogChangeClientEventsOnClient, Event);
+
+		// If we receive this from a different endpoint, it's probably from malicious user or someone scripting around with the API, either way warn about it.
+		if (!ensure(Context.SourceEndpointId == LiveSession->GetSessionServerEndpointId()))
+		{
+			UE_LOG(LogConcert, Warning, TEXT("The FConcertReplication_ChangeClientEvent is only supposed to be sent by the server."))
+			return;
+		}
+
+		const FConcertReplication_ChangeStream_Request& StreamChange = Event.StreamChange;
+		PredictAndApplyStreamChangeRemovedObjects(StreamChange);
+		FinalizePredictedStreamChange(StreamChange);
+
+		const FConcertReplication_ChangeAuthority_Request& AuthorityChange = Event.AuthorityChange;
+		ApplyAuthorityChangeRemovedObjects(AuthorityChange);
+		FinalizePredictedAuthorityChange(AuthorityChange, {}, Event.SyncControlChange);
+	}
+}
+
diff --git a/Engine/Plugins/Developer/Concert/ConcertSync/ConcertSyncClient/Source/ConcertSyncClient/Private/Replication/Manager/ReplicationManagerUtils.cpp b/Engine/Plugins/Developer/Concert/ConcertSync/ConcertSyncClient/Source/ConcertSyncClient/Private/Replication/Manager/ReplicationManagerUtils.cpp
deleted file mode 100644
diff --git a/Engine/Plugins/Developer/Concert/ConcertSync/ConcertSyncClient/Source/ConcertSyncClient/Private/Replication/Manager/ReplicationManagerUtils.h b/Engine/Plugins/Developer/Concert/ConcertSync/ConcertSyncClient/Source/ConcertSyncClient/Private/Replication/Manager/ReplicationManagerUtils.h
deleted file mode 100644
diff --git a/Engine/Plugins/Developer/Concert/ConcertSync/ConcertSyncClient/Source/ConcertSyncClient/Private/Replication/Manager/Utils/LocalSyncControl.h b/Engine/Plugins/Developer/Concert/ConcertSync/ConcertSyncClient/Source/ConcertSyncClient/Private/Replication/Manager/Utils/LocalSyncControl.h
new file mode 100644
--- /dev/null
+++ b/Engine/Plugins/Developer/Concert/ConcertSync/ConcertSyncClient/Source/ConcertSyncClient/Private/Replication/Manager/Utils/LocalSyncControl.h
@@ -0,0 +1,175 @@
+﻿// Copyright Epic Games, Inc. All Rights Reserved.
+
+#pragma once
+
+#include "IConcertSession.h"
+#include "Replication/Messages/SyncControl.h"
+
+#include "HAL/Platform.h"
+#include "Replication/SyncControlState.h"
+#include "Templates/UnrealTemplate.h"
+
+namespace UE::ConcertSyncClient::Replication
+{
+	/**
+	 * Receives network messages that explicitly and implicitly change sync control.
+	 * Those network messages are fed into FSyncControlState, which parses the messages.
+	 * 
+	 * Sync control is the set of objects this client is allowed to replicate.
+	 */
+	class FLocalSyncControl
+		: public FNoncopyable
+		, ConcertSyncCore::Replication::FSyncControlState
+	{
+		using Detail = FSyncControlState;
+	public:
+
+		explicit FLocalSyncControl(IConcertSession& InSession UE_LIFETIMEBOUND)
+			: Session(InSession)
+		{
+			Session.RegisterCustomEventHandler<FConcertReplication_ChangeSyncControl>(this, &FLocalSyncControl::OnChangeSyncControl);
+		}
+
+		~FLocalSyncControl()
+		{
+			Session.UnregisterCustomEventHandler<FConcertReplication_ChangeSyncControl>(this);
+		}
+
+		using FSyncControlState::IsObjectAllowed;
+		using FSyncControlState::Num;
+		using FSyncControlState::EnumerateAllowedObjects;
+		using FSyncControlState::EnumerateChanges;
+		using FPredictedObjectRemoval = FSyncControlState::FPredictedObjectRemoval;
+		// No using for FSyncControlState::Aggregate because we want to limit usage to below ProcessXChange functions
+
+		/** Stores the given change in this data structure. */
+		void ProcessSyncControlChange(const FConcertReplication_ChangeSyncControl& Event)
+		{
+			OnPreSyncControlChangedDelegate.Broadcast();
+			AppendChanges(Event);
+			OnPostSyncControlChangedDelegate.Broadcast();
+		}
+
+		/** Applies implicit and explicit changes to the client's sync control resulting from a completed authority change request and response. */
+		void ProcessAuthorityChange(const FConcertReplication_ChangeAuthority_Request& Request, const FConcertReplication_ChangeSyncControl& Response)
+		{
+			OnPreSyncControlChangedDelegate.Broadcast();
+			AppendAuthorityChange(Request, Response);
+			OnPostSyncControlChangedDelegate.Broadcast();
+		}
+
+		/**
+		 * Applies implicit changes to the client's sync control resulting from losing authority from objects removed from the stream.
+		 * You must validate that the request has also been accepted by the server!
+		 */
+		void ProcessStreamChange(const FConcertReplication_ChangeStream_Request& Request)
+		{
+			OnPreSyncControlChangedDelegate.Broadcast();
+			AppendStreamChange(Request);
+			OnPostSyncControlChangedDelegate.Broadcast();
+		}
+
+		/**
+		 * Applies the implicit changes made by the request assuming the request will be accepted.
+		 * @return The removed objects to be passed to ApplyOrRevertMuteResponse.
+		 */
+		FPredictedObjectRemoval PredictAndApplyMuteChanges(const FConcertReplication_ChangeMuteState_Request& Request)
+		{
+			bool bMadeChange = false;
+			FPredictedObjectRemoval Predicition = Detail::PredictAndApplyMuteChanges(Request, [this, &bMadeChange](const FConcertObjectInStreamID&)
+			{
+				if (!bMadeChange)
+				{
+					OnPreSyncControlChangedDelegate.Broadcast();
+				}
+				bMadeChange = true;
+			});
+			if (bMadeChange)
+			{
+				OnPostSyncControlChangedDelegate.Broadcast();
+			}
+
+			return Predicition;
+		}
+
+		/** Either reverts previous changes made if the request was rejected, or applies the sync control returned by the server otherwise. */
+		void ApplyOrRevertMuteResponse(const FPredictedObjectRemoval& RemovedByRequest, const FConcertReplication_ChangeMuteState_Response& Response)
+		{
+			bool bMadeChange = false;
+			Detail::ApplyOrRevertMuteResponse(RemovedByRequest, Response, [this, &bMadeChange](const FConcertObjectInStreamID&)
+			{
+				if (!bMadeChange)
+				{
+					OnPreSyncControlChangedDelegate.Broadcast();
+				}
+				bMadeChange = true;
+			});
+			if (bMadeChange)
+			{
+				OnPostSyncControlChangedDelegate.Broadcast();
+			}
+		}
+
+		
+		/**
+		 * Applies the implicit changes made by the request assuming the request will be accepted.
+		 * @return The removed objects to be passed to ApplyOrRevertRestoreContentResponse.
+		 */
+		FPredictedObjectRemoval PredictAndApplyRestoreContentChanges(const FConcertReplication_RestoreContent_Request& Request)
+		{
+			bool bMadeChange = false;
+			FPredictedObjectRemoval Predicition = Detail::PredictAndApplyRestoreContentChanges(Request, [this, &bMadeChange](const FConcertObjectInStreamID&)
+			{
+				if (!bMadeChange)
+				{
+					OnPreSyncControlChangedDelegate.Broadcast();
+				}
+				bMadeChange = true;
+			});
+			if (bMadeChange)
+			{
+				OnPostSyncControlChangedDelegate.Broadcast();
+			}
+
+			return Predicition;
+		}
+
+		/** Either reverts previous changes made if the request was rejected, or applies the sync control returned by the server otherwise. */
+		void ApplyOrRevertRestoreContentResponse(const FPredictedObjectRemoval& RemovedByRequest, const FConcertReplication_RestoreContent_Response& Response)
+		{
+			bool bMadeChange = false;
+			const auto OnChange = [this, &bMadeChange](const FConcertObjectInStreamID&)
+			{
+				if (!bMadeChange)
+				{
+					OnPreSyncControlChangedDelegate.Broadcast();
+				}
+				bMadeChange = true;
+			};
+			Detail::ApplyOrRevertRestoreContentResponse(RemovedByRequest, Response, OnChange, OnChange);
+			
+			if (bMadeChange)
+			{
+				OnPostSyncControlChangedDelegate.Broadcast();
+			}
+		}
+
+		void PredictAndApplyPutStateChanges(const FConcertReplication_RestoreContent_Request& Request);
+		
+		
+		DECLARE_MULTICAST_DELEGATE(FSyncControlChanged);
+		FSyncControlChanged& OnPreSyncControlChanged() { return OnPreSyncControlChangedDelegate; }
+		FSyncControlChanged& OnPostSyncControlChanged() { return OnPostSyncControlChangedDelegate; }
+
+	private:
+
+		/** The session to receive sync control changes on. */
+		IConcertSession& Session;
+		
+		FSyncControlChanged OnPreSyncControlChangedDelegate;
+		FSyncControlChanged OnPostSyncControlChangedDelegate;
+
+		void OnChangeSyncControl(const FConcertSessionContext&, const FConcertReplication_ChangeSyncControl& Event) { ProcessSyncControlChange(Event); }
+	};
+}
+
diff --git a/Engine/Plugins/Developer/Concert/ConcertSync/ConcertSyncClient/Source/ConcertSyncClient/Private/Replication/Manager/Utils/NetworkMessageLogging.cpp b/Engine/Plugins/Developer/Concert/ConcertSync/ConcertSyncClient/Source/ConcertSyncClient/Private/Replication/Manager/Utils/NetworkMessageLogging.cpp
new file mode 100644
--- /dev/null
+++ b/Engine/Plugins/Developer/Concert/ConcertSync/ConcertSyncClient/Source/ConcertSyncClient/Private/Replication/Manager/Utils/NetworkMessageLogging.cpp
@@ -0,0 +1,65 @@
+﻿// Copyright Epic Games, Inc. All Rights Reserved.
+
+#include "NetworkMessageLogging.h"
+
+namespace UE::ConcertSyncClient::Replication
+{
+	TAutoConsoleVariable<bool> CVarSimulateAuthorityTimeouts(
+		TEXT("Concert.Replication.SimulateAuthorityTimeouts"),
+		false,
+		TEXT("Whether the client should pretend that authority requests timed out instead of sending to the server.")
+		);
+	TAutoConsoleVariable<bool> CVarSimulateQueryTimeouts(
+		TEXT("Concert.Replication.SimulateQueryTimeouts"),
+		false,
+		TEXT("Whether the client should pretend that query requests timed out instead of sending to the server.")
+		);
+	TAutoConsoleVariable<bool> CVarSimulateStreamChangeTimeouts(
+		TEXT("Concert.Replication.SimulateStreamChangeTimeouts"),
+		false,
+		TEXT("Whether the client should pretend that stream change requests timed out instead of sending to the server.")
+		);
+
+	TAutoConsoleVariable<bool> CVarSimulateAuthorityRejection(
+		TEXT("Concert.Replication.SimulateAuthorityRejection"),
+		false,
+		TEXT("Whether the client should pretend that authority change requests were rejected.")
+		);
+	TAutoConsoleVariable<bool> CVarSimulateMuteRequestRejection(
+		TEXT("Concert.Replication.SimulateMuteRejection"),
+		false,
+		TEXT("Whether the client should pretend that mute change requests were rejected.")
+		);
+
+	TAutoConsoleVariable<bool> CVarLogStreamRequestsAndResponsesOnClient(
+		TEXT("Concert.Replication.LogStreamRequestsAndResponsesOnClient"),
+		false,
+		TEXT("Whether to log changes to streams.")
+		);
+	TAutoConsoleVariable<bool> CVarLogAuthorityRequestsAndResponsesOnClient(
+		TEXT("Concert.Replication.LogAuthorityRequestsAndResponsesOnClient"),
+		false,
+		TEXT("Whether to log changes to authority.")
+		);
+	TAutoConsoleVariable<bool> CVarLogMuteRequestsAndResponsesOnClient(
+		TEXT("Concert.Replication.LogMuteRequestsAndResponsesOnClient"),
+		false,
+		TEXT("Whether to log changes to the mute state.")
+		);
+	TAutoConsoleVariable<bool> CVarLogRestoreContentRequestsAndResponsesOnClient(
+		TEXT("Concert.Replication.LogRestoreContentRequestsAndResponsesOnClient"),
+		false,
+		TEXT("Whether to log restore content requests and responses.")
+		);
+		
+	TAutoConsoleVariable<bool> CVarLogChangeClientsRequestsAndResponsesOnClient(
+	TEXT("Concert.Replication.LogChangeClientsRequestsAndResponsesOnClient"),
+	false,
+	TEXT("Whether to log requests and responses that change multiple clients in one go.")
+	);
+	TAutoConsoleVariable<bool> CVarLogChangeClientEventsOnClient(
+		TEXT("Concert.Replication.LogChangeClientEventsOnClient"),
+		false,
+		TEXT("Whether to log messages from the server that notify us that the client's content has changed.")
+		);
+}
\ No newline at end of file
diff --git a/Engine/Plugins/Developer/Concert/ConcertSync/ConcertSyncClient/Source/ConcertSyncClient/Private/Replication/Manager/Utils/NetworkMessageLogging.h b/Engine/Plugins/Developer/Concert/ConcertSync/ConcertSyncClient/Source/ConcertSyncClient/Private/Replication/Manager/Utils/NetworkMessageLogging.h
new file mode 100644
--- /dev/null
+++ b/Engine/Plugins/Developer/Concert/ConcertSync/ConcertSyncClient/Source/ConcertSyncClient/Private/Replication/Manager/Utils/NetworkMessageLogging.h
@@ -0,0 +1,47 @@
+﻿// Copyright Epic Games, Inc. All Rights Reserved.
+
+#pragma once
+
+#include "ConcertLogGlobal.h"
+
+#include "Containers/UnrealString.h"
+#include "HAL/IConsoleManager.h"
+#include "JsonObjectConverter.h"
+
+namespace UE::ConcertSyncClient::Replication
+{
+	extern TAutoConsoleVariable<bool> CVarSimulateAuthorityTimeouts;
+	/** Whether the client should pretend that query requests timed out instead of sending to the server. */
+	extern TAutoConsoleVariable<bool> CVarSimulateQueryTimeouts;
+	/** Whether the client should pretend that stream change requests timed out instead of sending to the server. */
+	extern TAutoConsoleVariable<bool> CVarSimulateStreamChangeTimeouts;
+
+	/** Whether the client should pretend that authority change requests were rejected. */
+	extern TAutoConsoleVariable<bool> CVarSimulateAuthorityRejection;
+	/** Whether the client should pretend that mute change requests were rejected. */
+	extern TAutoConsoleVariable<bool> CVarSimulateMuteRequestRejection;
+
+	/** Whether to log changes to streams. */
+	extern TAutoConsoleVariable<bool> CVarLogStreamRequestsAndResponsesOnClient;
+	/** Whether to log changes to authority. */
+	extern TAutoConsoleVariable<bool> CVarLogAuthorityRequestsAndResponsesOnClient;
+	/** Whether to log changes to the mute state. */
+	extern TAutoConsoleVariable<bool> CVarLogMuteRequestsAndResponsesOnClient;
+	/** Whether to log restore content requests and responses. */
+	extern TAutoConsoleVariable<bool> CVarLogRestoreContentRequestsAndResponsesOnClient;
+	/** Whether to log requests and responses that change multiple clients in one go. */
+	extern TAutoConsoleVariable<bool> CVarLogChangeClientsRequestsAndResponsesOnClient;
+	/** Whether to log messages from the server that notify us that the client's content has changed. */
+	extern TAutoConsoleVariable<bool> CVarLogChangeClientEventsOnClient;
+	
+	template<typename TMessage>
+	static void LogNetworkMessage(const TAutoConsoleVariable<bool>& ShouldLog, const TMessage& Message)
+	{
+		if (ShouldLog.GetValueOnAnyThread())
+		{
+			FString JsonString;
+			FJsonObjectConverter::UStructToJsonObjectString(TMessage::StaticStruct(), &Message, JsonString, 0, 0);
+			UE_LOG(LogConcert, Log, TEXT("%s\n%s"), *TMessage::StaticStruct()->GetName(), *JsonString);
+		}
+	}
+}
diff --git a/Engine/Plugins/Developer/Concert/ConcertSync/ConcertSyncClient/Source/ConcertSyncClient/Private/Replication/Manager/Utils/ReplicationManagerUtils.cpp b/Engine/Plugins/Developer/Concert/ConcertSync/ConcertSyncClient/Source/ConcertSyncClient/Private/Replication/Manager/Utils/ReplicationManagerUtils.cpp
new file mode 100644
--- /dev/null
+++ b/Engine/Plugins/Developer/Concert/ConcertSync/ConcertSyncClient/Source/ConcertSyncClient/Private/Replication/Manager/Utils/ReplicationManagerUtils.cpp
@@ -0,0 +1,42 @@
+// Copyright Epic Games, Inc. All Rights Reserved.
+
+#include "ReplicationManagerUtils.h"
+
+#include "Replication/IConcertClientReplicationManager.h"
+
+#include "Algo/Transform.h"
+#include "Containers/Set.h"
+
+namespace UE::ConcertSyncClient::Replication
+{
+	TFuture<FConcertReplication_ChangeAuthority_Response> RejectAll(FConcertReplication_ChangeAuthority_Request&& Args)
+	{
+		return MakeFulfilledPromise<FConcertReplication_ChangeAuthority_Response>(
+			FConcertReplication_ChangeAuthority_Response{ EReplicationResponseErrorCode::Handled, MoveTemp(Args.TakeAuthority) }
+			).GetFuture();
+	}
+
+	TMap<FSoftObjectPath, TArray<FGuid>> ComputeRemovedObjects(
+		const TConstArrayView<FConcertReplicationStream> RegisteredStreams,
+		const FConcertReplication_ChangeStream_Request& Request)
+	{
+		TMap<FSoftObjectPath, TArray<FGuid>> BundledRemovedObjects;
+		for (const FConcertObjectInStreamID& RemovedObject : Request.ObjectsToRemove)
+		{
+			BundledRemovedObjects.FindOrAdd(RemovedObject.Object).AddUnique(RemovedObject.StreamId);
+		}
+		for (const FConcertReplicationStream& Stream : RegisteredStreams)
+		{
+			const FGuid& StreamId = Stream.BaseDescription.Identifier;
+			if (!Request.StreamsToRemove.Contains(StreamId))
+			{
+				continue;
+			}
+			for (const TPair<FSoftObjectPath, FConcertReplicatedObjectInfo>& Pair : Stream.BaseDescription.ReplicationMap.ReplicatedObjects)
+			{
+				BundledRemovedObjects.FindOrAdd(Pair.Key).AddUnique(StreamId);
+			}
+		}
+		return BundledRemovedObjects;
+	}
+}
diff --git a/Engine/Plugins/Developer/Concert/ConcertSync/ConcertSyncClient/Source/ConcertSyncClient/Private/Replication/Manager/Utils/ReplicationManagerUtils.h b/Engine/Plugins/Developer/Concert/ConcertSync/ConcertSyncClient/Source/ConcertSyncClient/Private/Replication/Manager/Utils/ReplicationManagerUtils.h
new file mode 100644
--- /dev/null
+++ b/Engine/Plugins/Developer/Concert/ConcertSync/ConcertSyncClient/Source/ConcertSyncClient/Private/Replication/Manager/Utils/ReplicationManagerUtils.h
@@ -0,0 +1,22 @@
+// Copyright Epic Games, Inc. All Rights Reserved.
+
+#pragma once
+
+#include "Async/Future.h"
+
+struct FConcertReplicationStream;
+struct FConcertReplication_ChangeStream_Request;
+struct FConcertReplication_ChangeAuthority_Response;
+struct FConcertReplication_ChangeAuthority_Request;
+
+namespace UE::ConcertSyncClient::Replication
+{
+	/** Creates a fulfilled future which rejects all items in FConcertChangeAuthority_Request::TakeAuthority. */
+	TFuture<FConcertReplication_ChangeAuthority_Response> RejectAll(FConcertReplication_ChangeAuthority_Request&& Args);
+
+	/** Enumerates the objects that will be removed by Request. */
+	TMap<FSoftObjectPath, TArray<FGuid>> ComputeRemovedObjects(
+		const TConstArrayView<FConcertReplicationStream> RegisteredStreams,
+		const FConcertReplication_ChangeStream_Request& Request
+		);
+};
diff --git a/Engine/Plugins/Developer/Concert/ConcertSync/ConcertSyncClient/Source/ConcertSyncClient/Private/Replication/Misc/LocalSyncControl.h b/Engine/Plugins/Developer/Concert/ConcertSync/ConcertSyncClient/Source/ConcertSyncClient/Private/Replication/Misc/LocalSyncControl.h
deleted file mode 100644
diff --git a/Engine/Plugins/Developer/Concert/ConcertSync/ConcertSyncClient/Source/ConcertSyncClient/Private/Replication/Processing/ClientReplicationDataCollector.cpp b/Engine/Plugins/Developer/Concert/ConcertSync/ConcertSyncClient/Source/ConcertSyncClient/Private/Replication/Processing/ClientReplicationDataCollector.cpp
--- a/Engine/Plugins/Developer/Concert/ConcertSync/ConcertSyncClient/Source/ConcertSyncClient/Private/Replication/Processing/ClientReplicationDataCollector.cpp
+++ b/Engine/Plugins/Developer/Concert/ConcertSync/ConcertSyncClient/Source/ConcertSyncClient/Private/Replication/Processing/ClientReplicationDataCollector.cpp
@@ -3,14 +3,14 @@
 #include "ClientReplicationDataCollector.h"
 
 #include "Replication/Data/ReplicationStream.h"
+#include "Replication/Data/ObjectIds.h"
 #include "Replication/Formats/IObjectReplicationFormat.h"
 #include "Replication/IConcertClientReplicationBridge.h"
+#include "Replication/Manager/Utils/LocalSyncControl.h"
 #include "Replication/ReplicationPropertyFilter.h"
-#include "Replication/Data/ObjectIds.h"
 
 #include "Algo/RemoveIf.h"
 #include "Misc/EBreakBehavior.h"
-#include "Replication/Misc/LocalSyncControl.h"
 
 namespace UE::ConcertSyncClient::Replication
 {
diff --git a/Engine/Plugins/Developer/Concert/ConcertSync/ConcertSyncClient/Source/ConcertSyncClient/Public/Replication/IConcertClientReplicationManager.h b/Engine/Plugins/Developer/Concert/ConcertSync/ConcertSyncClient/Source/ConcertSyncClient/Public/Replication/IConcertClientReplicationManager.h
--- a/Engine/Plugins/Developer/Concert/ConcertSync/ConcertSyncClient/Source/ConcertSyncClient/Public/Replication/IConcertClientReplicationManager.h
+++ b/Engine/Plugins/Developer/Concert/ConcertSync/ConcertSyncClient/Source/ConcertSyncClient/Public/Replication/IConcertClientReplicationManager.h
@@ -4,6 +4,7 @@
 
 #include "Misc/EBreakBehavior.h"
 #include "Replication/Messages/ChangeAuthority.h"
+#include "Replication/Messages/PutState.h"
 #include "Replication/Messages/ChangeStream.h"
 #include "Replication/Messages/ClientQuery.h"
 #include "Replication/Messages/Handshake.h"
@@ -134,9 +135,15 @@ public:
 	virtual TFuture<FConcertReplication_QueryMuteState_Response> QueryMuteState(FConcertReplication_QueryMuteState_Request Request = {}) = 0;
 	TFuture<FConcertReplication_QueryMuteState_Response> QueryMuteState(TSet<FSoftObjectPath> Objects);
 
-	/** Restore this client's stream content and authority to what a client had when they left. */
+	/** Restore this client's stream content and authority to what a client had when they left. Only works if session has EConcertSyncSessionFlags::ShouldEnableReplicationActivities set. */
 	virtual TFuture<FConcertReplication_RestoreContent_Response> RestoreContent(FConcertReplication_RestoreContent_Request Request = {}) = 0;
 
+	/**
+	 * Changes multiple clients' stream, authority, and optionally the global mute state. Useful for applying a preset.
+	 * Only works if session has EConcertSyncSessionFlags::ShouldEnableRemoteEditing set.
+	 */
+	virtual TFuture<FConcertReplication_PutState_Response> PutClientState(FConcertReplication_PutState_Request Request) = 0;
+
 	DECLARE_MULTICAST_DELEGATE(FOnPreStreamsChanged);
 	/** Called right before the result of GetRegisteredStreams changes. */
 	virtual FOnPreStreamsChanged& OnPreStreamsChanged() = 0;
diff --git a/Engine/Plugins/Developer/Concert/ConcertSync/ConcertSyncCore/Source/ConcertSyncCore/Public/ConcertSyncSessionFlags.h b/Engine/Plugins/Developer/Concert/ConcertSync/ConcertSyncCore/Source/ConcertSyncCore/Public/ConcertSyncSessionFlags.h
--- a/Engine/Plugins/Developer/Concert/ConcertSync/ConcertSyncCore/Source/ConcertSyncCore/Public/ConcertSyncSessionFlags.h
+++ b/Engine/Plugins/Developer/Concert/ConcertSync/ConcertSyncCore/Source/ConcertSyncCore/Public/ConcertSyncSessionFlags.h
@@ -54,13 +54,15 @@ enum class EConcertSyncSessionFlags : uint32
 	ShouldAllowGlobalMuting = 1<<14,
 	/** Whether EConcertSyncActivityEventType::Replication activities should be generated. */
 	ShouldEnableReplicationActivities = 1<<15,
+	/** Whether clients are allowed to edit each other's streams & authority (via FConcertReplication_PutState_Request). */
+	ShouldEnableRemoteEditing = 1 << 16,
 
 	/**
 	 * Defaults for different scenarios
 	 */
 	
 	/** Default settings for a multi-user session */
-	Default_MultiUserSession = EnableLiveSync | EnableConnectionHistory | EnableLocking | EnableTransactions | EnablePackages | EnableSequencer | EnablePresence | ShouldSendTransactionSnapshots | ShouldUsePackageSandbox | EnableReplication | ShouldAllowGlobalMuting | ShouldEnableReplicationActivities,
+	Default_MultiUserSession = EnableLiveSync | EnableConnectionHistory | EnableLocking | EnableTransactions | EnablePackages | EnableSequencer | EnablePresence | ShouldSendTransactionSnapshots | ShouldUsePackageSandbox | EnableReplication | ShouldAllowGlobalMuting | ShouldEnableReplicationActivities | ShouldEnableRemoteEditing,
 	/** Default settings for a disaster recovery session */
 	Default_DisasterRecoverySession = EnableTransactions | EnablePackages | EnableFileSharing | ShouldDiscardTransactionsOnPackageUnload | ShouldSendPackagePristineState | ShouldSendPackageAutoSaves,
 };
diff --git a/Engine/Plugins/Developer/Concert/ConcertSync/ConcertSyncCore/Source/ConcertSyncCore/Public/Replication/Data/AuthorityConflict.h b/Engine/Plugins/Developer/Concert/ConcertSync/ConcertSyncCore/Source/ConcertSyncCore/Public/Replication/Data/AuthorityConflict.h
new file mode 100644
--- /dev/null
+++ b/Engine/Plugins/Developer/Concert/ConcertSync/ConcertSyncCore/Source/ConcertSyncCore/Public/Replication/Data/AuthorityConflict.h
@@ -0,0 +1,48 @@
+﻿// Copyright Epic Games, Inc. All Rights Reserved.
+
+#pragma once
+
+#include "Containers/Array.h"
+#include "Containers/Set.h"
+#include "ObjectIds.h"
+#include "AuthorityConflict.generated.h"
+
+/** Describes an authority conflict for two objects. */
+USTRUCT()
+struct FConcertAuthorityConflict
+{
+	GENERATED_BODY()
+
+	/** The object that was attempted to take authority over but could not be. */
+	UPROPERTY()
+	FConcertObjectInStreamID AttemptedObject;
+
+	/** The object of another client that Source conflicts with. */
+	UPROPERTY()
+	FConcertReplicatedObjectId ConflictingObject;
+};
+
+USTRUCT()
+struct FConcertAuthorityConflictArray
+{
+	GENERATED_BODY()
+
+	/** Each AttemptedObject is unique. */
+	UPROPERTY()
+	TArray<FConcertAuthorityConflict> Conflicts;
+
+	FConcertAuthorityConflict& FindOrAdd(const FConcertObjectInStreamID& AttemptedObject)
+	{
+		FConcertAuthorityConflict* Result = Conflicts.FindByPredicate([&AttemptedObject](const FConcertAuthorityConflict& Conflict)
+		{
+			return Conflict.AttemptedObject == AttemptedObject;
+		});
+		if (Result)
+		{
+			return *Result;
+		}
+
+		const int32 Index = Conflicts.Add({ AttemptedObject });
+		return Conflicts[Index];
+	}
+};
\ No newline at end of file
diff --git a/Engine/Plugins/Developer/Concert/ConcertSync/ConcertSyncCore/Source/ConcertSyncCore/Public/Replication/Data/ObjectIds.h b/Engine/Plugins/Developer/Concert/ConcertSync/ConcertSyncCore/Source/ConcertSyncCore/Public/Replication/Data/ObjectIds.h
--- a/Engine/Plugins/Developer/Concert/ConcertSync/ConcertSyncCore/Source/ConcertSyncCore/Public/Replication/Data/ObjectIds.h
+++ b/Engine/Plugins/Developer/Concert/ConcertSync/ConcertSyncCore/Source/ConcertSyncCore/Public/Replication/Data/ObjectIds.h
@@ -56,5 +56,14 @@ struct FConcertReplicatedObjectId : public FConcertObjectInStreamID
 	}
 };
 
+USTRUCT()
+struct FConcertObjectInStreamArray
+{
+	GENERATED_BODY()
+
+	UPROPERTY()
+	TArray<FConcertObjectInStreamID> Objects;
+};
+
 CONCERTSYNCCORE_API uint32 GetTypeHash(const FConcertObjectInStreamID& StreamObject);
 CONCERTSYNCCORE_API uint32 GetTypeHash(const FConcertReplicatedObjectId& StreamObject);
diff --git a/Engine/Plugins/Developer/Concert/ConcertSync/ConcertSyncCore/Source/ConcertSyncCore/Public/Replication/Data/ReplicationStreamArray.h b/Engine/Plugins/Developer/Concert/ConcertSync/ConcertSyncCore/Source/ConcertSyncCore/Public/Replication/Data/ReplicationStreamArray.h
new file mode 100644
--- /dev/null
+++ b/Engine/Plugins/Developer/Concert/ConcertSync/ConcertSyncCore/Source/ConcertSyncCore/Public/Replication/Data/ReplicationStreamArray.h
@@ -0,0 +1,15 @@
+﻿// Copyright Epic Games, Inc. All Rights Reserved.
+
+#pragma once
+
+#include "ReplicationStream.h"
+#include "ReplicationStreamArray.generated.h"
+
+USTRUCT()
+struct FConcertReplicationStreamArray
+{
+	GENERATED_BODY()
+
+	UPROPERTY()
+	TArray<FConcertReplicationStream> Streams;
+};
diff --git a/Engine/Plugins/Developer/Concert/ConcertSync/ConcertSyncCore/Source/ConcertSyncCore/Public/Replication/Messages/ChangeClientEvent.h b/Engine/Plugins/Developer/Concert/ConcertSync/ConcertSyncCore/Source/ConcertSyncCore/Public/Replication/Messages/ChangeClientEvent.h
new file mode 100644
--- /dev/null
+++ b/Engine/Plugins/Developer/Concert/ConcertSync/ConcertSyncCore/Source/ConcertSyncCore/Public/Replication/Messages/ChangeClientEvent.h
@@ -0,0 +1,32 @@
+﻿// Copyright Epic Games, Inc. All Rights Reserved.
+
+#pragma once
+
+#include "ChangeAuthority.h"
+#include "ChangeStream.h"
+#include "ChangeClientEvent.generated.h"
+
+/**
+ * Sent by server to notify a client that their stream content and / or authority has been changed by an external entity,
+ * i.e. when the change was not initiated by the client itself.
+ * 
+ * For now, this is only in response to a remote client sending a FConcertReplication_PutState_Request but in the future it could also be
+ * the user changing client content via the server UI.
+ */
+USTRUCT()
+struct FConcertReplication_ChangeClientEvent
+{
+	GENERATED_BODY()
+
+	/** The change made to the client's streams. */
+	UPROPERTY()
+	FConcertReplication_ChangeStream_Request StreamChange;
+
+	/** The change made to the client's authority. */
+	UPROPERTY()
+	FConcertReplication_ChangeAuthority_Request AuthorityChange;
+
+	/** The change made to the client's sync control in response to the above changes. */
+	UPROPERTY()
+	FConcertReplication_ChangeSyncControl SyncControlChange;
+};
\ No newline at end of file
diff --git a/Engine/Plugins/Developer/Concert/ConcertSync/ConcertSyncCore/Source/ConcertSyncCore/Public/Replication/Messages/ChangeStream.h b/Engine/Plugins/Developer/Concert/ConcertSync/ConcertSyncCore/Source/ConcertSyncCore/Public/Replication/Messages/ChangeStream.h
--- a/Engine/Plugins/Developer/Concert/ConcertSync/ConcertSyncCore/Source/ConcertSyncCore/Public/Replication/Messages/ChangeStream.h
+++ b/Engine/Plugins/Developer/Concert/ConcertSync/ConcertSyncCore/Source/ConcertSyncCore/Public/Replication/Messages/ChangeStream.h
@@ -294,7 +294,7 @@ struct FConcertReplication_ChangeStream_Response
 	UPROPERTY()
 	TMap<FConcertObjectInStreamID, EConcertPutObjectErrorCode> ObjectsToPutSemanticErrors;
 	
-	/** Streams that were in StreamsToAdd but that were not created. */
+	/** Streams that were in StreamsToAdd but that were not created (stream was already registered). */
 	UPROPERTY()
 	TSet<FGuid> FailedStreamCreation;
 
diff --git a/Engine/Plugins/Developer/Concert/ConcertSync/ConcertSyncCore/Source/ConcertSyncCore/Public/Replication/Messages/PutState.h b/Engine/Plugins/Developer/Concert/ConcertSync/ConcertSyncCore/Source/ConcertSyncCore/Public/Replication/Messages/PutState.h
new file mode 100644
--- /dev/null
+++ b/Engine/Plugins/Developer/Concert/ConcertSync/ConcertSyncCore/Source/ConcertSyncCore/Public/Replication/Messages/PutState.h
@@ -0,0 +1,152 @@
+﻿// Copyright Epic Games, Inc. All Rights Reserved.
+
+#pragma once
+
+#include "Muting.h"
+#include "Replication/Data/AuthorityConflict.h"
+#include "Replication/Data/ReplicationStreamArray.h"
+#include "PutState.generated.h"
+
+UENUM(Flags) 
+enum class EConcertReplicationPutStateFlags : uint8
+{
+	None,
+
+	/** If another client has disconnected between the time the request was sent and received by the client, do not fail the request. */
+	SkipDisconnectedClients = 1 << 0,
+
+	/** Prior to applying the FConcertReplication_PutState_Request::MuteChange, clear all mute state. */
+	ResetMuteState = 1 << 1,
+
+	Default = SkipDisconnectedClients
+};
+ENUM_CLASS_FLAGS(EConcertReplicationPutStateFlags);
+
+/**
+ * Requests that sets the state of replication atomically.
+ * It can be used to change the streams and / or authority of multiple clients together with the mute state in a single request.
+ * This is useful i.e. for setting the session to a preset state.
+ * 
+ * Upon successful execution, the server will send a FConcertReplication_ChangeClientEvent notification to all affected clients except the requesting
+ * client to update them about the changed information; the requesting client receives only sync control changes (since they know the rest they changed).
+ * 
+ * Upon failure, it is guaranteed no changes have taken place.
+ * The request fails whenever any of the following occurs:
+ *	- applying the request would generate any authority conflict.
+ *	For example, a request that wants client 1 and 2 to both replicate the same properties of the same object would fail.
+ *	- a client leaves replication before the request is processed and the SkipDisconnectedClients is not set.
+ *
+ * This request is only available if EConcertSyncSessionFlags::ShouldEnableRemoteEditing.
+ */
+USTRUCT()
+struct FConcertReplication_PutState_Request
+{
+	GENERATED_BODY()
+
+	/** Additional flags that alter the behavior of the request. */
+	UPROPERTY()
+	EConcertReplicationPutStateFlags Flags = EConcertReplicationPutStateFlags::Default;
+
+	/** Maps a client's endpoint ID to the stream content it should have. */
+	UPROPERTY()
+	TMap<FGuid, FConcertReplicationStreamArray> NewStreams;
+
+	/** Maps a client's endpoint ID to new authority it should have. */
+	UPROPERTY()
+	TMap<FGuid, FConcertObjectInStreamArray> NewAuthorityState;
+
+	/** Mute state to apply after NewStreams and NewAuthorityState has been applied. */
+	UPROPERTY()
+	FConcertReplication_ChangeMuteState_Request MuteChange;
+};
+
+/**
+ * Indicates success or failure of the FConcertReplication_PutState_Request.
+ * 
+ * While a request may have multiple errors, the errors are returned in declaration order of this enum.
+ * Hence, the response code is the first thing that fails when reading top to bottom.
+ */
+UENUM()
+enum class EConcertReplicationPutStateResponseCode : uint8
+{
+	/** The request timed out. */
+	Timeout,
+	/** The request was executed successfully. */
+	Success,
+
+	/**
+	 * Either:
+	 * - Session does not allow remote editing (EConcertSyncSessionFlags::ShouldEnableRemoteEditing is disabled).
+	 * - Session does not allow muting (EConcertSyncSessionFlags::ShouldAllowGlobalMuting is disabled) and a mute request was made.
+	 */
+	FeatureDisabled,
+	
+	/**
+	 * No changes have been made.
+	 * 
+	 * Only returned if the flag EConcertReplicationChangeClientsFlags:: SkipDisconnectedClients was not set and
+	 * FConcertReplication_PutState_Request::StreamChanges or FConcertReplication_PutState_Request::AuthorityChanges referenced a client
+	 * endpoint not in the session or that has not joined replication.
+	 * 
+	 * FConcertReplication_PutState_Response::UnknownEndpoints contains the unknown endpoints.
+	 */
+	ClientUnknown,
+
+	/**
+	 * No changes have been made. 
+	 * The StreamChanges was produced an error incorrect, e.g. if a stream change references an unknown stream
+	 * or if a mute request references an object unknown to the server.
+	 */
+	StreamError,
+	
+	/**
+	 * No changes have been made. 
+	 * The request failed because the AuthorityChanges portion would generate a conflict with other clients if applied.
+	 */
+	AuthorityConflict,
+
+	/**
+	 * No changes have been made.
+	 * The request failed because 
+	 */
+	MuteError
+};
+
+USTRUCT()
+struct FConcertReplication_PutState_Response
+{
+	GENERATED_BODY()
+
+	UPROPERTY()
+	EConcertReplicationPutStateResponseCode ResponseCode = EConcertReplicationPutStateResponseCode::Timeout;
+
+	/**
+	 * Maps a key from AuthorityChanges to the errors it had.
+	 * Set if ResponseCode == EConcertReplicationChangeClientsResponseCode::AuthorityConflict.
+	 */
+	UPROPERTY()
+	TMap<FGuid, FConcertAuthorityConflictArray> AuthorityChangeConflicts;
+
+	/** Sets if ResponseCode == EConcertReplicationChangeClientsResponseCode::ClientUnknown. */
+	UPROPERTY()
+	TSet<FGuid> UnknownEndpoints;
+
+	/**
+	 * The sync control changes that happened to the requesting client.
+	 *
+	 * Generally, this does not contain changes that the client can predict locally, i.e.
+	 * 
+	 * - Retaining sync control.
+	 *	If the requesting client had sync control for an object (before the request) and the request does not modify it, that object is not contained in SyncControl.
+	 *	For example, if the client has sync control over Foo, and NewStreams & NewAuthorityState still contain Foo, Foo is not listed in SyncControl (since the client can predict the change).
+	 *	
+	 * - Muting:
+	 *	If the requesting client had sync control for an object (before the request) and the request mutes it, that object is not contained in SyncControl.
+	 *	For example, if the client has sync control over Foo, and MuteChanges mutes Foo, Foo is not listed in SyncControl (since the client can predict the change).
+	 *	
+	 */
+	UPROPERTY()
+	FConcertReplication_ChangeSyncControl SyncControl;
+
+	bool IsSuccess() const { return ResponseCode == EConcertReplicationPutStateResponseCode::Success; }
+};
\ No newline at end of file
diff --git a/Engine/Plugins/Developer/Concert/ConcertSync/ConcertSyncCore/Source/ConcertSyncCore/Public/Replication/Misc/SyncControlUtils.h b/Engine/Plugins/Developer/Concert/ConcertSync/ConcertSyncCore/Source/ConcertSyncCore/Public/Replication/Misc/SyncControlUtils.h
new file mode 100644
--- /dev/null
+++ b/Engine/Plugins/Developer/Concert/ConcertSync/ConcertSyncCore/Source/ConcertSyncCore/Public/Replication/Misc/SyncControlUtils.h
@@ -0,0 +1,42 @@
+﻿// Copyright Epic Games, Inc. All Rights Reserved.
+
+#pragma once
+
+#include "Replication/Data/ObjectIds.h"
+#include "Replication/Messages/SyncControl.h"
+
+#include "Containers/Map.h"
+#include "Misc/EnumClassFlags.h"
+
+namespace UE::ConcertSyncCore
+{
+	enum class EAppendSyncControlFlags : uint8
+	{
+		None,
+		/** The resulting SyncControlToUpdate should not contain any false values. */
+		SkipLostControl = 1 << 0
+	};
+	ENUM_CLASS_FLAGS(EAppendSyncControlFlags);
+
+	/** Appends AppendedSyncControl to SyncControlToUpdate overriding old entries. */
+	void AppendSyncControl(FConcertReplication_ChangeSyncControl& SyncControlToUpdate, const FConcertReplication_ChangeSyncControl& AppendedSyncControl, EAppendSyncControlFlags Flags = EAppendSyncControlFlags::None);
+}
+
+namespace UE::ConcertSyncCore
+{
+	inline void AppendSyncControl(FConcertReplication_ChangeSyncControl& SyncControlToUpdate, const FConcertReplication_ChangeSyncControl& AppendedSyncControl, EAppendSyncControlFlags Flags)
+	{
+		const bool bShouldSkipLostControl = EnumHasAnyFlags(Flags, EAppendSyncControlFlags::SkipLostControl);
+		for (const TPair<FConcertObjectInStreamID, bool>& Change : AppendedSyncControl.NewControlStates)
+		{
+			if (bShouldSkipLostControl && !Change.Value)
+			{
+				// Do not append the change but more importantly, also remove any potential true state from SyncControlToUpdate.
+				SyncControlToUpdate.NewControlStates.Remove(Change.Key);
+				continue;
+			}
+			
+			SyncControlToUpdate.NewControlStates.Add(Change.Key, Change.Value);
+		}
+	}
+}
diff --git a/Engine/Plugins/Developer/Concert/ConcertSync/ConcertSyncCore/Source/ConcertSyncCore/Public/Replication/SyncControlState.h b/Engine/Plugins/Developer/Concert/ConcertSync/ConcertSyncCore/Source/ConcertSyncCore/Public/Replication/SyncControlState.h
--- a/Engine/Plugins/Developer/Concert/ConcertSync/ConcertSyncCore/Source/ConcertSyncCore/Public/Replication/SyncControlState.h
+++ b/Engine/Plugins/Developer/Concert/ConcertSync/ConcertSyncCore/Source/ConcertSyncCore/Public/Replication/SyncControlState.h
@@ -67,8 +67,8 @@ namespace UE::ConcertSyncCore::Replication
 
 		/** Combines implicit and explicit sync control changes caused by authority changes. */
 		template<CObjectInStreamInvocable TOnAllowed, CObjectInStreamInvocable TOnDisallowed>
-		void AppendAuthorityChange(const FConcertReplication_ChangeAuthority_Request& Request, const FConcertReplication_ChangeAuthority_Response& Response, TOnAllowed&& OnAllowed, TOnDisallowed&& OnDisallowed);
-		void AppendAuthorityChange(const FConcertReplication_ChangeAuthority_Request& Request, const FConcertReplication_ChangeAuthority_Response& Response);
+		void AppendAuthorityChange(const FConcertReplication_ChangeAuthority_Request& Request, const FConcertReplication_ChangeSyncControl& Response, TOnAllowed&& OnAllowed, TOnDisallowed&& OnDisallowed);
+		void AppendAuthorityChange(const FConcertReplication_ChangeAuthority_Request& Request, const FConcertReplication_ChangeSyncControl& Response);
 
 		/** Combines implicit sync control changes caused by stream change. */
 		template<CObjectInStreamInvocable TOnDisallowed>
@@ -163,7 +163,7 @@ namespace UE::ConcertSyncCore::Replication
 	}
 	
 	template <CObjectInStreamInvocable TOnAllowed, CObjectInStreamInvocable TOnDisallowed>
-	void FSyncControlState::AppendAuthorityChange(const FConcertReplication_ChangeAuthority_Request& Request, const FConcertReplication_ChangeAuthority_Response& Response, TOnAllowed&& OnAllowed, TOnDisallowed&& OnDisallowed)
+	void FSyncControlState::AppendAuthorityChange(const FConcertReplication_ChangeAuthority_Request& Request, const FConcertReplication_ChangeSyncControl& Response, TOnAllowed&& OnAllowed, TOnDisallowed&& OnDisallowed)
 	{
 		for (const TPair<FSoftObjectPath, FConcertStreamArray>& ImplicitChange : Request.ReleaseAuthority)
 		{
@@ -179,10 +179,10 @@ namespace UE::ConcertSyncCore::Replication
 			}
 		}
 
-		AppendChanges(Response.SyncControl, MoveTemp(OnAllowed), MoveTemp(OnDisallowed));
+		AppendChanges(Response, MoveTemp(OnAllowed), MoveTemp(OnDisallowed));
 	}
 
-	inline void FSyncControlState::AppendAuthorityChange(const FConcertReplication_ChangeAuthority_Request& Request, const FConcertReplication_ChangeAuthority_Response& Response)
+	inline void FSyncControlState::AppendAuthorityChange(const FConcertReplication_ChangeAuthority_Request& Request, const FConcertReplication_ChangeSyncControl& Response)
 	{
 		AppendAuthorityChange(Request, Response, [](const FConcertObjectInStreamID&){}, [](const FConcertObjectInStreamID&){});
 	}
@@ -283,7 +283,7 @@ namespace UE::ConcertSyncCore::Replication
 				OnAllowed,
 				[](const FConcertObjectInStreamID&)
 				{
-					ensureAlwaysMsgf(false, TEXT("By contract, objects losing sync control are not supposed to be listed here. FConcertReplication_ChangeMuteState_Response::SyncControl docoumentation."));
+					ensureAlwaysMsgf(false, TEXT("By contract, objects losing sync control are not supposed to be listed here. FConcertReplication_ChangeMuteState_Response::SyncControl documentation."));
 				});
 		}
 		else
diff --git a/Engine/Plugins/Developer/Concert/ConcertSync/ConcertSyncServer/Source/ConcertSyncServer/Private/Replication/AuthorityManager.cpp b/Engine/Plugins/Developer/Concert/ConcertSync/ConcertSyncServer/Source/ConcertSyncServer/Private/Replication/AuthorityManager.cpp
--- a/Engine/Plugins/Developer/Concert/ConcertSync/ConcertSyncServer/Source/ConcertSyncServer/Private/Replication/AuthorityManager.cpp
+++ b/Engine/Plugins/Developer/Concert/ConcertSync/ConcertSyncServer/Source/ConcertSyncServer/Private/Replication/AuthorityManager.cpp
@@ -8,6 +8,7 @@
 #include "Replication/Data/ObjectIds.h"
 #include "Replication/Data/ReplicationStream.h"
 #include "Replication/Messages/ChangeAuthority.h"
+#include "Util/GroundTruthOverride.h"
 
 namespace UE::ConcertSyncServer::Replication
 {
@@ -56,6 +57,40 @@ namespace UE::ConcertSyncServer::Replication
 				}
 			}
 		}
+
+		const FConcertReplicationStream* FindClientStreamById(const IStreamEnumerator& Getters, const FGuid& ClientId, const FGuid& StreamId)
+		{
+			const FConcertReplicationStream* StreamDescription = nullptr;
+			Getters.ForEachStream(ClientId, [&StreamId, &StreamDescription](const FConcertReplicationStream& Stream) mutable
+			{
+				if (Stream.BaseDescription.Identifier == StreamId)
+				{
+					StreamDescription = &Stream;
+					return EBreakBehavior::Break;
+				}
+				return EBreakBehavior::Continue;
+			});
+			return StreamDescription;
+		}
+
+		const FConcertObjectReplicationMap* FindClientReplicationMapById(
+			const ConcertSyncCore::Replication::AuthorityConflictUtils::IReplicationGroundTruth& Getters,
+			const FGuid& ClientId,
+			const FGuid& StreamId
+			)
+		{
+			const FConcertObjectReplicationMap* ReplicationMap = nullptr;
+			Getters.ForEachStream(ClientId, [&StreamId, &ReplicationMap](const FGuid& InStreamId, const FConcertObjectReplicationMap& InReplicationMap) mutable
+			{
+				if (InStreamId == InStreamId)
+				{
+					ReplicationMap = &InReplicationMap;
+					return EBreakBehavior::Break;
+				}
+				return EBreakBehavior::Continue;
+			});
+			return ReplicationMap;
+		}
 	}
 	
 	FAuthorityManager::FAuthorityManager(
@@ -128,7 +163,7 @@ namespace UE::ConcertSyncServer::Replication
 		const FConcertPropertySelection* PropertiesToCheck = OverwriteProperties;
 		if (!PropertiesToCheck)
 		{
-			const FConcertReplicationStream* Description = FindClientStreamById(ClientId, Object.StreamId);
+			const FConcertReplicationStream* Description = Private::FindClientStreamById(Getters, ClientId, Object.StreamId);
 			const FConcertReplicatedObjectInfo* PropertyInfo = Description ? Description->BaseDescription.ReplicationMap.ReplicatedObjects.Find(Object.Object) : nullptr;
 			PropertiesToCheck = PropertyInfo ? &PropertyInfo->PropertySelection : nullptr;
 		}
@@ -158,6 +193,38 @@ namespace UE::ConcertSyncServer::Replication
 		return EnumerateAuthorityConflicts(Object) == EAuthorityResult::Allowed;
 	}
 
+	FAuthorityManager::EAuthorityResult FAuthorityManager::EnumerateAuthorityConflictsWithOverrides(
+		const FConcertReplicatedObjectId& Object,
+		const TMap<FGuid, FConcertReplicationStreamArray>& StreamOverrides,
+		const TMap<FGuid, FConcertObjectInStreamArray>& AuthorityOverrides,
+		FProcessAuthorityConflict ProcessConflict
+		) const
+	{
+		const FClientId& ClientId = Object.SenderEndpointId;
+		const FGroundTruthOverride GroundTruth(StreamOverrides, AuthorityOverrides, Getters, *this);
+		
+		const FConcertObjectReplicationMap* ReplicationMap = Private::FindClientReplicationMapById(GroundTruth, ClientId, Object.StreamId);
+		const FConcertReplicatedObjectInfo* PropertyInfo = ReplicationMap ? ReplicationMap->ReplicatedObjects.Find(Object.Object) : nullptr;
+		const FConcertPropertySelection* PropertiesToCheck = PropertyInfo ? &PropertyInfo->PropertySelection : nullptr;
+		if (!PropertiesToCheck)
+		{
+			return EAuthorityResult::NoRegisteredProperties;
+		}
+
+		using namespace ConcertSyncCore::Replication;
+		using namespace ConcertSyncCore::Replication::AuthorityConflictUtils;
+		const EAuthorityConflict Conflict = AuthorityConflictUtils::EnumerateAuthorityConflicts(
+			Object.SenderEndpointId,
+			Object.Object,
+			PropertiesToCheck->ReplicatedProperties,
+			GroundTruth,
+			[&ProcessConflict](const FClientId& ClientId, const FStreamId& StreamId, const FConcertPropertyChain& Property)
+			{
+				return ProcessConflict(ClientId, StreamId, Property);
+			});
+		return Conflict == EAuthorityConflict::Allowed ? EAuthorityResult::Allowed : EAuthorityResult::Conflict;
+	}
+
 	void FAuthorityManager::OnPostClientLeft(const FClientId& ClientEndpointId)
 	{
 		ClientAuthorityData.Remove(ClientEndpointId);
@@ -259,19 +326,4 @@ namespace UE::ConcertSyncServer::Replication
 			OutChangedSyncControl = GenerateSyncControlDelegate.Execute(ClientId);
 		}
 	}
-	
-	const FConcertReplicationStream* FAuthorityManager::FindClientStreamById(const FClientId& ClientId, const FStreamId& StreamId) const
-	{
-		const FConcertReplicationStream* StreamDescription = nullptr;
-		Getters.ForEachStream(ClientId, [&StreamId, &StreamDescription](const FConcertReplicationStream& Stream) mutable
-		{
-			if (Stream.BaseDescription.Identifier == StreamId)
-			{
-				StreamDescription = &Stream;
-				return EBreakBehavior::Break;
-			}
-			return EBreakBehavior::Continue;
-		});
-		return StreamDescription;
-	}
 }
diff --git a/Engine/Plugins/Developer/Concert/ConcertSync/ConcertSyncServer/Source/ConcertSyncServer/Private/Replication/AuthorityManager.h b/Engine/Plugins/Developer/Concert/ConcertSync/ConcertSyncServer/Source/ConcertSyncServer/Private/Replication/AuthorityManager.h
--- a/Engine/Plugins/Developer/Concert/ConcertSync/ConcertSyncServer/Source/ConcertSyncServer/Private/Replication/AuthorityManager.h
+++ b/Engine/Plugins/Developer/Concert/ConcertSync/ConcertSyncServer/Source/ConcertSyncServer/Private/Replication/AuthorityManager.h
@@ -10,6 +10,7 @@
 #include "Misc/EBreakBehavior.h"
 #include "Templates/Function.h"
 
+struct FConcertReplicationStreamArray;
 struct FConcertStreamArray;
 struct FConcertReplication_ChangeSyncControl;
 class IConcertSession;
@@ -20,6 +21,7 @@ struct FConcertPropertyChain;
 struct FConcertPropertySelection;
 struct FConcertSessionContext;
 struct FConcertObjectInStreamID;
+struct FConcertObjectInStreamArray;
 struct FConcertReplicatedObjectId;
 struct FConcertReplicationStream;
 
@@ -75,6 +77,21 @@ namespace UE::ConcertSyncServer::Replication
 		/** Whether it is legal for this the client identified by ClientId to take control over the object given the stream the client has registered. */
 		bool CanTakeAuthority(const FConcertReplicatedObjectId& Object) const;
 
+		/**
+		 * Enumerates all authority conflicts, if any, that would occur if ObjectChange.SenderEndpointId were to take authority over the identified object.
+		 *
+		 * While EnumerateAuthorityConflicts allows you to specify OverwriteProperties, this version allows you to override the
+		 * 
+		 * @return Whether there were any conflicts
+		 */
+
+		EAuthorityResult EnumerateAuthorityConflictsWithOverrides(
+			const FConcertReplicatedObjectId& Object,
+			const TMap<FGuid, FConcertReplicationStreamArray>& StreamOverrides,
+			const TMap<FGuid, FConcertObjectInStreamArray>& AuthorityOverrides,
+			FProcessAuthorityConflict ProcessConflict = [](auto&, auto&, auto&){ return EBreakBehavior::Break; }
+			) const;
+
 		/** Notifies this manager that the client has left, which means all their authority is now gone. */
 		void OnPostClientLeft(const FClientId& ClientEndpointId);
 		/** Takes away authority from the given client from the given object. */
@@ -119,9 +136,6 @@ namespace UE::ConcertSyncServer::Replication
 			FConcertReplication_ChangeSyncControl& OutChangedSyncControl,
 			bool bShouldLog = true
 			);
-
-		/** Finds a stream registered with the client by its ID. */
-		const FConcertReplicationStream* FindClientStreamById(const FClientId& ClientId, const FStreamId& StreamId) const;
 	};
 }
 
diff --git a/Engine/Plugins/Developer/Concert/ConcertSync/ConcertSyncServer/Source/ConcertSyncServer/Private/Replication/ConcertServerReplicationManager.cpp b/Engine/Plugins/Developer/Concert/ConcertSync/ConcertSyncServer/Source/ConcertSyncServer/Private/Replication/ConcertServerReplicationManager.cpp
--- a/Engine/Plugins/Developer/Concert/ConcertSync/ConcertSyncServer/Source/ConcertSyncServer/Private/Replication/ConcertServerReplicationManager.cpp
+++ b/Engine/Plugins/Developer/Concert/ConcertSync/ConcertSyncServer/Source/ConcertSyncServer/Private/Replication/ConcertServerReplicationManager.cpp
@@ -10,6 +10,7 @@
 #include "Replication/Data/ClientQueriedInfo.h"
 #include "Replication/Formats/FullObjectFormat.h"
 #include "Replication/IReplicationWorkspace.h"
+#include "Replication/Messages/PutState.h"
 #include "Replication/Messages/ChangeStream.h"
 #include "Replication/Messages/ClientQuery.h"
 #include "Replication/Messages/Handshake.h"
@@ -46,6 +47,7 @@ namespace UE::ConcertSyncServer::Replication
 		Session->RegisterCustomRequestHandler<FConcertReplication_QueryReplicationInfo_Request, FConcertReplication_QueryReplicationInfo_Response>(this, &FConcertServerReplicationManager::HandleQueryReplicationInfoRequest);
 		Session->RegisterCustomRequestHandler<FConcertReplication_ChangeStream_Request, FConcertReplication_ChangeStream_Response>(this, &FConcertServerReplicationManager::HandleChangeStreamRequest);
 		Session->RegisterCustomRequestHandler<FConcertReplication_RestoreContent_Request, FConcertReplication_RestoreContent_Response>(this, &FConcertServerReplicationManager::HandleRestoreContentRequest);
+		Session->RegisterCustomRequestHandler<FConcertReplication_PutState_Request, FConcertReplication_PutState_Response>(this, &FConcertServerReplicationManager::HandlePutStateRequest);
 		Session->RegisterCustomEventHandler<FConcertReplication_LeaveEvent>(this, &FConcertServerReplicationManager::HandleLeaveReplicationSessionRequest);
 		
 		Session->OnSessionClientChanged().AddRaw(this, &FConcertServerReplicationManager::OnConnectionChanged);
diff --git a/Engine/Plugins/Developer/Concert/ConcertSync/ConcertSyncServer/Source/ConcertSyncServer/Private/Replication/ConcertServerReplicationManager.h b/Engine/Plugins/Developer/Concert/ConcertSync/ConcertSyncServer/Source/ConcertSyncServer/Private/Replication/ConcertServerReplicationManager.h
--- a/Engine/Plugins/Developer/Concert/ConcertSync/ConcertSyncServer/Source/ConcertSyncServer/Private/Replication/ConcertServerReplicationManager.h
+++ b/Engine/Plugins/Developer/Concert/ConcertSync/ConcertSyncServer/Source/ConcertSyncServer/Private/Replication/ConcertServerReplicationManager.h
@@ -6,7 +6,7 @@
 #include "ConcertMessages.h"
 #include "ConcertReplicationClient.h"
 #include "Enumeration/IRegistrationEnumerator.h"
-#include "MuteManager.h"
+#include "Muting/MuteManager.h"
 #include "Replication/Formats/IObjectReplicationFormat.h"
 #include "Replication/IConcertServerReplicationManager.h"
 #include "Replication/Messages/Handshake.h"
@@ -16,11 +16,13 @@
 #include "SyncControlManager.h"
 
 #include "HAL/Platform.h"
+#include "Replication/Messages/PutState.h"
 #include "Replication/Messages/RestoreContent.h"
 #include "Templates/SharedPointer.h"
 #include "Templates/Tuple.h"
 #include "Templates/UnrealTemplate.h"
 
+struct FConcertReplication_ChangeClientEvent;
 class FConcertServerWorkspace;
 class IConcertClientReplicationBridge;
 class IConcertServerSession;
@@ -140,6 +142,27 @@ namespace UE::ConcertSyncServer::Replication
 		/** Restores the mute state of Client. */
 		void RestoreMuteState(const FConcertReplicationClient& Client, FConcertReplication_ChangeSyncControl& OutChangedSyncControl);
 		void ApplyRestoringMuteRequest(const FConcertReplicationClient& Client, const FConcertReplication_ChangeMuteState_Request& AggregatedRequest, FConcertReplication_ChangeSyncControl& OutChangedSyncControl);
+
+		// Changing multiple clients in one go
+		EConcertSessionResponseCode HandlePutStateRequest(const FConcertSessionContext& Context, const FConcertReplication_PutState_Request& Request, FConcertReplication_PutState_Response& Response);
+		void ApplyPutStateRequest(
+			const FGuid& RequestingEndpointId,
+			const FConcertReplication_PutState_Request& Request,
+			const TMap<FGuid, FConcertReplication_ChangeStream_Request> StreamRequests,
+			FConcertReplication_PutState_Response& Response
+			);
+		void ApplyPutState_Streams(
+			const FGuid& RequestingEndpointId,
+			const TMap<FGuid, FConcertReplication_ChangeStream_Request> StreamRequests,
+			TMap<FGuid, FConcertReplication_ChangeClientEvent>& ClientChanges
+			);
+		void ApplyPutState_Authority(
+			const FGuid& RequestingEndpointId,
+			const TSet<FConcertObjectInStreamID> RequestingClientSyncControlBefore,
+			const FConcertReplication_PutState_Request& Request,
+			FConcertReplication_PutState_Response& Response,
+			TMap<FGuid, FConcertReplication_ChangeClientEvent>& ClientChanges
+			);
 		
 		// Leaving
 		void HandleLeaveReplicationSessionRequest(const FConcertSessionContext& ConcertSessionContext, const FConcertReplication_LeaveEvent& EventData);
@@ -147,7 +170,7 @@ namespace UE::ConcertSyncServer::Replication
 
 		/** Cleans up the client's replication state after leaving. */
 		void OnClientLeftReplication(const FGuid& EndpointId);
-		/** Calls FConcertServerWorkspace::AddReplicationActivity with the current client state so it can be restored upon re-joining. */
+		/** Calls FConcertServerWorkspace::AddReplicationActivity with the current client state, so it can be restored upon re-joining. */
 		void ProduceClientLeftActivity(const FConcertReplicationClient& Client) const;
 
 		/** Generates an activity for a mute request */
diff --git a/Engine/Plugins/Developer/Concert/ConcertSync/ConcertSyncServer/Source/ConcertSyncServer/Private/Replication/ConcertServerReplicationManager_ChangeStream.cpp b/Engine/Plugins/Developer/Concert/ConcertSync/ConcertSyncServer/Source/ConcertSyncServer/Private/Replication/ConcertServerReplicationManager_ChangeStream.cpp
--- a/Engine/Plugins/Developer/Concert/ConcertSync/ConcertSyncServer/Source/ConcertSyncServer/Private/Replication/ConcertServerReplicationManager_ChangeStream.cpp
+++ b/Engine/Plugins/Developer/Concert/ConcertSync/ConcertSyncServer/Source/ConcertSyncServer/Private/Replication/ConcertServerReplicationManager_ChangeStream.cpp
@@ -8,139 +8,11 @@
 #include "Replication/ChangeStreamSharedUtils.h"
 #include "Replication/ConcertReplicationClient.h"
 #include "Replication/Messages/ChangeStream.h"
-#include "Replication/Messages/Handshake.h"
 #include "Util/LogUtils.h"
+#include "Util/StreamChangeValidation.h"
 
 namespace UE::ConcertSyncServer::Replication
 {
-	namespace Private
-	{
-		static const FConcertReplicationStream* FindExistingStream(const FConcertReplicationClient& Client, const FGuid& StreamId)
-		{
-			return Client.GetStreamDescriptions().FindByPredicate([&StreamId](const FConcertReplicationStream& Description)
-			{
-				return Description.BaseDescription.Identifier == StreamId;
-			});
-		}
-		
-		/** Validates that ObjectsToPut writes only to pre-existing streams. */
-		static void ValidatePutObjectsRequestSemantics(const FConcertReplication_ChangeStream_Request& Request, const FConcertReplicationClient& Client, FConcertReplication_ChangeStream_Response& OutResponse)
-		{
-			const auto PutObjectHasEnoughData = [](const FConcertReplicationStream* ExistingStream, const FSoftObjectPath& ChangedObjectPath, const FConcertReplication_ChangeStream_PutObject& PutObject)
-			{
-				const bool bHasProperties = !PutObject.Properties.ReplicatedProperties.IsEmpty();
-				const bool bHasClassPath = !PutObject.ClassPath.IsNull();
-				const bool bIsEditingExistingObjectDefinition = ExistingStream->BaseDescription.ReplicationMap.ReplicatedObjects.Contains(ChangedObjectPath);
-				const bool bHasEnoughData = (!bIsEditingExistingObjectDefinition && bHasProperties && bHasClassPath)
-					|| (bIsEditingExistingObjectDefinition && bHasProperties)
-					|| (bIsEditingExistingObjectDefinition && bHasClassPath);
-				return bHasEnoughData;
-			};
-			
-			for (const TPair<FConcertObjectInStreamID, FConcertReplication_ChangeStream_PutObject>& Change : Request.ObjectsToPut)
-			{
-				const FConcertObjectInStreamID ChangedObject = Change.Key;
-				
-				const FGuid& StreamToModify = ChangedObject.StreamId;
-				const FConcertReplicationStream* ExistingStream = FindExistingStream(Client, StreamToModify);
-				const bool bStreamExists = ExistingStream != nullptr;
-				if (bStreamExists)
-				{
-					// FConcertReplicatedObjectInfo must have non-empty values for its members. Hence we must check that
-					// 1. if creating a new entry, both fields are given
-					// 2. if writing to pre-existing entry, at least one field is given
-					const FConcertReplication_ChangeStream_PutObject& PutObject = Change.Value; 
-					if (!PutObjectHasEnoughData(ExistingStream, ChangedObject.Object, PutObject))
-					{
-						OutResponse.ObjectsToPutSemanticErrors.Add(ChangedObject, EConcertPutObjectErrorCode::MissingData);
-					}
-				}
-				else
-				{
-					// ObjectsToPut can only write to pre-existing streams
-					UE_LOG(LogConcert, Log, TEXT("Semantic error: unknown stream %s"), *StreamToModify.ToString(EGuidFormats::Short));
-					OutResponse.ObjectsToPutSemanticErrors.Add(ChangedObject, EConcertPutObjectErrorCode::UnresolvedStream);
-				}
-			}
-		}
-
-		/** Checks that StreamsToAdd do not conflict with pre-existing ones and that all IDs in the request are also unique. */
-		static void ValidateAddedStreamsAreUnique(const FConcertReplication_ChangeStream_Request& Request, const FConcertReplicationClient& Client, FConcertReplication_ChangeStream_Response& OutResponse)
-		{
-			TSet<FGuid> DuplicateEntryDetection;
-			for (const FConcertReplicationStream& NewStream : Request.StreamsToAdd)
-			{
-				// StreamsToAdd is invalid if there is already a stream with the same ID registered ...
-				const FGuid& NewStreamId = NewStream.BaseDescription.Identifier;
-				const bool bIdAlreadyExists = FindExistingStream(Client, NewStreamId) != nullptr;
-				const bool bIsStreamRemoved = Request.StreamsToRemove.Contains(NewStreamId);
-
-				// ... or StreamsToAdd contains the same ID multiple times
-				bool bIsDuplicateEntry = false;
-				DuplicateEntryDetection.FindOrAdd(NewStreamId, &bIsDuplicateEntry);
-				
-				if ((bIdAlreadyExists && !bIsStreamRemoved)
-					|| bIsDuplicateEntry)
-				{
-					UE_LOG(LogConcert, Log, TEXT("Duplicate stream entry %s"), *NewStreamId.ToString(EGuidFormats::Short));
-					OutResponse.FailedStreamCreation.Add(NewStreamId);
-				}
-			}
-		}
-
-		/** If the requesting client has authority over a changed object, checks that no other client has authority over the properties being added. */
-		static void LookForAuthorityConflicts(const FConcertReplication_ChangeStream_Request& Request, const FConcertReplicationClient& Client, const FAuthorityManager& AuthorityManager, FConcertReplication_ChangeStream_Response& OutResponse)
-		{
-			const FGuid ClientEndpointId = Client.GetClientEndpointId();
-			for (const TPair<FConcertObjectInStreamID, FConcertReplication_ChangeStream_PutObject>& PutObjectPair : Request.ObjectsToPut)
-			{
-				// No conflict possible if requesting client does not have authority over the changed object
-				const FConcertReplicatedObjectId ReplicatedObjectInfo { { PutObjectPair.Key }, ClientEndpointId };
-				if (!AuthorityManager.HasAuthorityToChange(ReplicatedObjectInfo))
-				{
-					continue;
-				}
-
-				// Requester not changing properties (must be changing ClassPath)? Also no conflict possible.
-				const FConcertPropertySelection& PropertySelection = PutObjectPair.Value.Properties;
-				const bool bIsEditingProperties = !PropertySelection.ReplicatedProperties.IsEmpty();
-				if (!bIsEditingProperties)
-				{
-					continue;
-				}
-
-				// Simply check whether any other client is already sending any of the requested properties.
-				AuthorityManager.EnumerateAuthorityConflicts(ReplicatedObjectInfo, &PropertySelection,
-					[&OutResponse, &ReplicatedObjectInfo](const FGuid& ClientId, const FGuid& StreamId, const FConcertPropertyChain& Property)
-					{
-						const FConcertReplicatedObjectId ConflictingObject = { { StreamId, ReplicatedObjectInfo.Object }, ClientId };
-						OutResponse.AuthorityConflicts.Add(ReplicatedObjectInfo, ConflictingObject);
-						
-						UE_LOG(LogConcert, Log, TEXT("Authority conflict with client %s for stream %s for property %s"), *ClientId.ToString(EGuidFormats::Short), *StreamId.ToString(EGuidFormats::Short), *Property.ToString());
-						return EBreakBehavior::Continue;
-					});
-			}
-		}
-
-		/** Checks whether this request is valid to apply. */
-		static bool ShouldAcceptRequest(
-			const FConcertReplication_ChangeStream_Request& Request,
-			const FConcertReplicationClient& Client,
-			const FAuthorityManager& AuthorityManager,
-			FConcertReplication_ChangeStream_Response& OutResponse
-			)
-		{
-			OutResponse.ErrorCode = EReplicationResponseErrorCode::Handled;
-			
-			ValidatePutObjectsRequestSemantics(Request, Client, OutResponse);
-			ValidateAddedStreamsAreUnique(Request, Client, OutResponse);
-			LookForAuthorityConflicts(Request, Client, AuthorityManager, OutResponse);
-			ConcertSyncCore::Replication::ChangeStreamUtils::ValidateFrequencyChanges(Request, Client.GetStreamDescriptions(), &OutResponse.FrequencyErrors);
-			
-			return OutResponse.IsSuccess();
-		}
-	}
-	
 	TAutoConsoleVariable<bool> CVarLogStreamRequestsAndResponsesOnServer(
 		TEXT("Concert.Replication.LogStreamRequestsAndResponsesOnServer"),
 		false,
@@ -157,7 +29,7 @@ namespace UE::ConcertSyncServer::Replication
 		
 		const FGuid SendingClientId = ConcertSessionContext.SourceEndpointId;
 		const TUniquePtr<FConcertReplicationClient>* SendingClient = Clients.Find(SendingClientId);
-		if (SendingClient && Private::ShouldAcceptRequest(Request, *SendingClient->Get(), AuthorityManager, Response))
+		if (SendingClient && ValidateStreamChangeRequest(*SendingClient->Get(), AuthorityManager, Request, Response))
 		{
 			ApplyChangeStreamRequest(Request, *SendingClient->Get());
 		}
@@ -170,7 +42,7 @@ namespace UE::ConcertSyncServer::Replication
 		LogNetworkMessage(CVarLogStreamRequestsAndResponsesOnServer, Response, [&](){ return GetClientName(*Session, ConcertSessionContext.SourceEndpointId); });
 		return EConcertSessionResponseCode::Success;
 	}
-	
+
 	void FConcertServerReplicationManager::ApplyChangeStreamRequest(const FConcertReplication_ChangeStream_Request& Request, FConcertReplicationClient& Client)
 	{
 		const FGuid& SendingClientId = Client.GetClientEndpointId();
diff --git a/Engine/Plugins/Developer/Concert/ConcertSync/ConcertSyncServer/Source/ConcertSyncServer/Private/Replication/ConcertServerReplicationManager_PutState.cpp b/Engine/Plugins/Developer/Concert/ConcertSync/ConcertSyncServer/Source/ConcertSyncServer/Private/Replication/ConcertServerReplicationManager_PutState.cpp
new file mode 100644
--- /dev/null
+++ b/Engine/Plugins/Developer/Concert/ConcertSync/ConcertSyncServer/Source/ConcertSyncServer/Private/Replication/ConcertServerReplicationManager_PutState.cpp
@@ -0,0 +1,373 @@
+﻿// Copyright Epic Games, Inc. All Rights Reserved.
+
+#include "ConcertLogGlobal.h"
+#include "ConcertServerReplicationManager.h"
+#include "Algo/AllOf.h"
+#include "ConcertServer/Private/ConcertServerSession.h"
+#include "Replication/Messages/ChangeClientEvent.h"
+
+#include "Containers/Array.h"
+#include "Containers/Set.h"
+#include "Containers/Map.h"
+#include "Muting/PredictedStateObjectHierarchy.h"
+#include "Replication/Misc/SyncControlUtils.h"
+#include "Util/StreamChangeValidation.h"
+
+#if UE_BUILD_DEBUG
+#include "JsonObjectConverter.h"
+#endif
+
+namespace UE::ConcertSyncServer::Replication
+{
+	namespace Private
+	{
+#if UE_BUILD_DEBUG
+		static FString RequestToString(const FConcertReplication_PutState_Request& Request)
+		{
+			FString JsonString;
+			FJsonObjectConverter::UStructToJsonObjectString(FConcertReplication_PutState_Request::StaticStruct(), &Request, JsonString, 0, 0);
+			return JsonString;
+		}
+#endif
+
+		static TMap<FGuid, FConcertReplication_ChangeStream_Request> BuildStreamRequests(
+			const TMap<FGuid, TUniquePtr<FConcertReplicationClient>>& Clients,
+			const FConcertReplication_PutState_Request& Request
+			)
+		{
+			TMap<FGuid, FConcertReplication_ChangeStream_Request> Result;
+
+			for (const TPair<FGuid, FConcertReplicationStreamArray>& Pair : Request.NewStreams)
+			{
+				const FGuid& EndpointId = Pair.Key;
+				const TUniquePtr<FConcertReplicationClient>* Client = Clients.Find(EndpointId);
+				if (!Client)
+				{
+					continue;
+				}
+
+				FConcertReplication_ChangeStream_Request& ChangeStream = Result.Add(EndpointId);
+				Algo::Transform(Client->Get()->GetStreamDescriptions(), ChangeStream.StreamsToRemove, [](const FConcertReplicationStream& Stream){ return Stream.BaseDescription.Identifier; });
+				ChangeStream.StreamsToAdd = Pair.Value.Streams;
+			}
+
+			return Result;
+		}
+
+		static FConcertReplication_ChangeAuthority_Request BuildAuthorityRequest(
+			const FGuid& EndpointId,
+			const TArray<FConcertReplicationStream>& Streams,
+			const FAuthorityManager& AuthorityManager,
+			const TArray<FConcertObjectInStreamID>& ObjectsToOwn
+			)
+		{
+			FConcertReplication_ChangeAuthority_Request AuthorityRequest;
+			for (const FConcertReplicationStream& Stream : Streams)
+			{
+				const FGuid& StreamId = Stream.BaseDescription.Identifier;
+				AuthorityManager.EnumerateAuthority(EndpointId, StreamId, [&ObjectsToOwn, &AuthorityRequest, &StreamId](const FSoftObjectPath& ObjectPath)
+				{
+					const FConcertObjectInStreamID ObjectId{ StreamId, ObjectPath };
+					if (!ObjectsToOwn.Contains(ObjectId))
+					{
+						AuthorityRequest.ReleaseAuthority.Add(ObjectPath).StreamIds.AddUnique(StreamId);
+					}
+					return EBreakBehavior::Continue;
+				});
+			}
+			for (const FConcertObjectInStreamID& ObjectId : ObjectsToOwn)
+			{
+				AuthorityRequest.TakeAuthority.FindOrAdd(ObjectId.Object).StreamIds.AddUnique(ObjectId.StreamId);
+			}
+			return AuthorityRequest;
+		}
+
+		static bool ValidateAllClientsKnown(
+			const TMap<FGuid, TUniquePtr<FConcertReplicationClient>>& Clients,
+			const FConcertReplication_PutState_Request& Request,
+			FConcertReplication_PutState_Response& Response
+			)
+		{
+			if (EnumHasAnyFlags(Request.Flags, EConcertReplicationPutStateFlags::SkipDisconnectedClients))
+			{
+				return true;
+			}
+			
+			bool bSuccess = true;
+			
+			for (const TPair<FGuid, FConcertReplicationStreamArray>& Pair : Request.NewStreams)
+			{
+				const FGuid& ClientId = Pair.Key;
+				if (!Clients.Contains(ClientId))
+				{
+					bSuccess = false;
+					Response.UnknownEndpoints.Add(ClientId);
+				}
+			}
+			
+			for (const TPair<FGuid, FConcertObjectInStreamArray>& Pair : Request.NewAuthorityState)
+			{
+				const FGuid& ClientId = Pair.Key;
+				if (!Clients.Contains(ClientId))
+				{
+					bSuccess = false;
+					Response.UnknownEndpoints.Add(ClientId);
+				}
+			}
+
+			Response.ResponseCode = bSuccess ? Response.ResponseCode : EConcertReplicationPutStateResponseCode::ClientUnknown;
+			return bSuccess;
+		}
+
+		static bool ValidateStreamsAreValid(
+			const TMap<FGuid, TUniquePtr<FConcertReplicationClient>>& Clients,
+			const FAuthorityManager& AuthorityManager,
+			const TMap<FGuid, FConcertReplication_ChangeStream_Request>& StreamRequests,
+			FConcertReplication_PutState_Response& Response
+			)
+		{
+			const bool bSuccess = Algo::AllOf(StreamRequests, [&Clients, &AuthorityManager](const TPair<FGuid, FConcertReplication_ChangeStream_Request>& Pair)
+			{
+				const TUniquePtr<FConcertReplicationClient>* Client = Clients.Find(Pair.Key);
+				return !Client || ValidateStreamChangeRequest(*Client->Get(), AuthorityManager, Pair.Value);
+			});
+
+			Response.ResponseCode = bSuccess ? Response.ResponseCode : EConcertReplicationPutStateResponseCode::StreamError;
+			return bSuccess;
+		}
+
+		/** This checks that Request generates no conflicts with any of the clients that are NOT modified by Request. */
+		static bool ValidateNoConflicts(
+			const FAuthorityManager& AuthorityManager,
+			const FConcertReplication_PutState_Request& Request,
+			FConcertReplication_PutState_Response& Response
+			)
+		{
+			bool bSuccess = true;
+			
+			for (const TPair<FGuid, FConcertObjectInStreamArray>& Pair : Request.NewAuthorityState)
+			{
+				const FGuid& ClientId = Pair.Key;
+
+				for (const FConcertObjectInStreamID& ObjectId : Pair.Value.Objects)
+				{
+					const FConcertReplicatedObjectId ReplicatedObjectId { { ObjectId }, ClientId };
+					const auto HandleConflict = [&Response, &ClientId, &ReplicatedObjectId](const FGuid& ConflictingClientId, const FGuid& StreamId, const FConcertPropertyChain& Property)
+					{
+						FConcertAuthorityConflictArray& Conflicts = Response.AuthorityChangeConflicts.FindOrAdd(ClientId);
+						Conflicts.FindOrAdd(ReplicatedObjectId).ConflictingObject = { { StreamId, ReplicatedObjectId.Object }, ConflictingClientId };
+						return EBreakBehavior::Continue;
+					};
+					
+					const FAuthorityManager::EAuthorityResult AuthorityResult = AuthorityManager.EnumerateAuthorityConflictsWithOverrides(
+						ReplicatedObjectId,
+						Request.NewStreams,
+						// Passing in NewAuthorityState causes us to get a conflict when the request specifies overlapping authority,
+						// e.g. if the request says client 1 and client 2 are supposed to replicate the same object properties. 
+						Request.NewAuthorityState,
+						HandleConflict
+						);
+					bSuccess &= AuthorityResult == FAuthorityManager::EAuthorityResult::Allowed;
+				}
+			}
+
+			Response.ResponseCode = bSuccess ? Response.ResponseCode : EConcertReplicationPutStateResponseCode::AuthorityConflict;
+			return bSuccess;
+		}
+
+		static bool ValidateMuteRequest(
+			const TMap<FGuid, TUniquePtr<FConcertReplicationClient>>& Clients,
+			const FMuteManager& MuteManager,
+			const EConcertSyncSessionFlags SessionFlags,
+			const FConcertReplication_PutState_Request& Request,
+			FConcertReplication_PutState_Response& Response
+			)
+		{
+			if (!EnumHasAnyFlags(SessionFlags, EConcertSyncSessionFlags::ShouldAllowGlobalMuting))
+			{
+				const bool bIsAllowed = Request.MuteChange.IsEmpty();
+				Response.ResponseCode = bIsAllowed ? Response.ResponseCode : EConcertReplicationPutStateResponseCode::FeatureDisabled;
+				return bIsAllowed;
+			}
+
+			FPredictedStateObjectHierarchy FutureHierarchy;
+			FutureHierarchy.AddClients(Request.NewStreams);
+			FutureHierarchy.AddClients(Clients, [&Request](const FGuid& ClientId){ return !Request.NewStreams.Contains(ClientId); });
+			const bool bIsRequestValid = MuteManager.ValidateRequest(Request.MuteChange, &FutureHierarchy);
+
+			Response.ResponseCode = bIsRequestValid ? Response.ResponseCode : EConcertReplicationPutStateResponseCode::MuteError;
+			return bIsRequestValid;
+		}
+		
+		static bool ValidatePutStateRequest(
+			const TMap<FGuid, TUniquePtr<FConcertReplicationClient>>& Clients,
+			const FAuthorityManager& AuthorityManager,
+			const FMuteManager& MuteManager,
+			const EConcertSyncSessionFlags SessionFlags,
+			const FConcertReplication_PutState_Request& Request,
+			const TMap<FGuid, FConcertReplication_ChangeStream_Request>& StreamRequests,
+			FConcertReplication_PutState_Response& Response
+			)
+		{
+			return ValidateAllClientsKnown(Clients, Request, Response)
+				&& ValidateStreamsAreValid(Clients, AuthorityManager, StreamRequests, Response)
+				&& ValidateNoConflicts(AuthorityManager, Request, Response)
+				&& ValidateMuteRequest(Clients, MuteManager, SessionFlags, Request, Response);
+		}
+
+		static void ApplyPutState_Mute(
+			FMuteManager& MuteManager,
+			const FGuid& RequestingEndpointId,
+			const FConcertReplication_ChangeMuteState_Request& Request,
+			FConcertReplication_PutState_Response& Response,
+			TMap<FGuid, FConcertReplication_ChangeClientEvent>& ClientChanges
+			)
+		{
+			const bool bSuccess = MuteManager.ApplyRequestAndEnumerateSyncControl(
+			Request,
+			[&RequestingEndpointId, &Response, &ClientChanges](const FGuid& EndpointId, const FConcertReplication_ChangeSyncControl& SyncControlChange)
+			{
+				if (SyncControlChange.IsEmpty())
+				{
+					return;
+				}
+
+				using namespace UE::ConcertSyncCore;
+				if (EndpointId == RequestingEndpointId)
+				{
+					AppendSyncControl(Response.SyncControl, SyncControlChange, EAppendSyncControlFlags::SkipLostControl);
+				}
+				else
+				{
+					AppendSyncControl(ClientChanges.FindOrAdd(EndpointId).SyncControlChange, SyncControlChange);
+				}
+			});
+			UE_CLOG(!bSuccess, LogConcert, Error, TEXT("Failed to apply mute request as part of put state request from %s"), *RequestingEndpointId.ToString());
+		}
+	}
+	
+	EConcertSessionResponseCode FConcertServerReplicationManager::HandlePutStateRequest(
+		const FConcertSessionContext& Context,
+		const FConcertReplication_PutState_Request& Request,
+		FConcertReplication_PutState_Response& Response
+		)
+	{
+		if (!EnumHasAnyFlags(SessionFlags, EConcertSyncSessionFlags::ShouldEnableRemoteEditing))
+		{
+			Response.ResponseCode = EConcertReplicationPutStateResponseCode::FeatureDisabled;
+			return EConcertSessionResponseCode::Success;
+		}
+
+		const TMap<FGuid, FConcertReplication_ChangeStream_Request> StreamRequests = Private::BuildStreamRequests(Clients, Request);
+		if (Private::ValidatePutStateRequest(Clients, AuthorityManager, MuteManager, SessionFlags, Request, StreamRequests, Response))
+		{
+			Response.ResponseCode = EConcertReplicationPutStateResponseCode::Success;
+			ApplyPutStateRequest(Context.SourceEndpointId, Request, StreamRequests, Response);
+		}
+
+		return EConcertSessionResponseCode::Success;
+	}
+
+	void FConcertServerReplicationManager::ApplyPutStateRequest(
+		const FGuid& RequestingEndpointId,
+		const FConcertReplication_PutState_Request& Request,
+		const TMap<FGuid, FConcertReplication_ChangeStream_Request> StreamRequests,
+		FConcertReplication_PutState_Response& Response
+		)
+	{
+		TMap<FGuid, FConcertReplication_ChangeClientEvent> ChangedSyncControl;
+		const TSet<FConcertObjectInStreamID> RequestingClientSyncControlBefore = [this, &RequestingEndpointId]()
+		{
+			const TSet<FConcertObjectInStreamID>* SyncControl = SyncControlManager.GetClientControlledObjects(RequestingEndpointId);
+			return SyncControl ? *SyncControl : TSet<FConcertObjectInStreamID>{};
+		}();
+		
+		ApplyPutState_Streams(RequestingEndpointId, StreamRequests, ChangedSyncControl);
+		ApplyPutState_Authority(RequestingEndpointId, RequestingClientSyncControlBefore, Request, Response, ChangedSyncControl);
+		if (EnumHasAnyFlags(SessionFlags, EConcertSyncSessionFlags::ShouldAllowGlobalMuting))
+		{
+			Private::ApplyPutState_Mute(MuteManager, RequestingEndpointId, Request.MuteChange, Response, ChangedSyncControl);
+		}
+
+		for (TPair<FGuid, FConcertReplication_ChangeClientEvent>& SyncControlPair : ChangedSyncControl)
+		{
+			const FGuid& EndpointId = SyncControlPair.Key;
+			// The requesting client receives the sync control via Response (which was written by ApplyChangeClients_StreamPart, etc.)
+			if (EndpointId != RequestingEndpointId)
+			{
+				Session->SendCustomEvent(SyncControlPair.Value, EndpointId, EConcertMessageFlags::ReliableOrdered);
+			}
+		}
+	}
+
+	void FConcertServerReplicationManager::ApplyPutState_Streams(
+		const FGuid& RequestingEndpointId,
+		const TMap<FGuid, FConcertReplication_ChangeStream_Request> StreamRequests,
+		TMap<FGuid, FConcertReplication_ChangeClientEvent>& ClientChanges
+		)
+	{
+		for (const TPair<FGuid, FConcertReplication_ChangeStream_Request>& StreamPair : StreamRequests)
+		{
+			const FGuid& EndpointId = StreamPair.Key;
+			
+			const FConcertReplication_ChangeStream_Request& StreamRequest = StreamPair.Value;
+			if (EndpointId != RequestingEndpointId)
+			{
+				ClientChanges.Add(EndpointId).StreamChange = StreamRequest;
+			}
+
+			// StreamRequests was created by BuildStreamRequests, which only adds entries for valid clients.
+			FConcertReplicationClient& Client = *Clients[EndpointId];
+			ApplyChangeStreamRequest(StreamRequest, Client);
+		}
+	}
+
+	void FConcertServerReplicationManager::ApplyPutState_Authority(
+		const FGuid& RequestingEndpointId,
+		const TSet<FConcertObjectInStreamID> RequestingClientSyncControlBefore,
+		const FConcertReplication_PutState_Request& Request,
+		FConcertReplication_PutState_Response& Response,
+		TMap<FGuid, FConcertReplication_ChangeClientEvent>& ClientChanges
+		)
+	{
+		for (const TPair<FGuid, FConcertObjectInStreamArray>& AuthorityPair : Request.NewAuthorityState)
+		{
+			const FGuid& EndpointId = AuthorityPair.Key;
+			
+			// EConcertReplicationChangeClientsFlags::SkipDisconnectedClients flag may be set.
+			const TUniquePtr<FConcertReplicationClient>* Client = Clients.Find(EndpointId);
+			if (!Client)
+			{
+				continue;
+			}
+
+			const TArray<FConcertObjectInStreamID> ObjectsToOwn = AuthorityPair.Value.Objects;
+			const FConcertReplication_ChangeAuthority_Request AuthorityRequest = Private::BuildAuthorityRequest(EndpointId, Client->Get()->GetStreamDescriptions(), AuthorityManager, ObjectsToOwn);
+			TMap<FSoftObjectPath, FConcertStreamArray> RejectedObjects;
+			if (EndpointId == RequestingEndpointId)
+			{
+				FConcertReplication_ChangeSyncControl& RequesterSyncControlChange = Response.SyncControl;
+				AuthorityManager.ApplyChangeAuthorityRequest(EndpointId, AuthorityRequest, RejectedObjects, Response.SyncControl);
+				// Only return actual changes of sync control to the requesting client.
+				for (auto It = RequesterSyncControlChange.NewControlStates.CreateIterator(); It; ++It)
+				{
+					if (RequestingClientSyncControlBefore.Contains(It->Key))
+					{
+						It.RemoveCurrent();
+					}
+				}
+			}
+			else
+			{
+				ClientChanges.FindOrAdd(EndpointId).AuthorityChange = AuthorityRequest;
+				FConcertReplication_ChangeSyncControl& RemoteSyncControlChange = ClientChanges.FindOrAdd(EndpointId).SyncControlChange;
+				AuthorityManager.ApplyChangeAuthorityRequest(EndpointId, AuthorityRequest, RejectedObjects, RemoteSyncControlChange);
+			}
+
+#if UE_BUILD_DEBUG
+			ensure(RejectedObjects.IsEmpty());
+			UE_CLOG(!RejectedObjects.IsEmpty(), LogConcert, Warning, TEXT("Authority portion for client %s was not validated correctly for request %s."), *EndpointId.ToString(), *Private::RequestToString(Request));
+#endif
+		}
+	}
+}
diff --git a/Engine/Plugins/Developer/Concert/ConcertSync/ConcertSyncServer/Source/ConcertSyncServer/Private/Replication/MuteManager.cpp b/Engine/Plugins/Developer/Concert/ConcertSync/ConcertSyncServer/Source/ConcertSyncServer/Private/Replication/MuteManager.cpp
deleted file mode 100644
diff --git a/Engine/Plugins/Developer/Concert/ConcertSync/ConcertSyncServer/Source/ConcertSyncServer/Private/Replication/MuteManager.h b/Engine/Plugins/Developer/Concert/ConcertSync/ConcertSyncServer/Source/ConcertSyncServer/Private/Replication/MuteManager.h
deleted file mode 100644
diff --git a/Engine/Plugins/Developer/Concert/ConcertSync/ConcertSyncServer/Source/ConcertSyncServer/Private/Replication/Muting/MuteManager.cpp b/Engine/Plugins/Developer/Concert/ConcertSync/ConcertSyncServer/Source/ConcertSyncServer/Private/Replication/Muting/MuteManager.cpp
new file mode 100644
--- /dev/null
+++ b/Engine/Plugins/Developer/Concert/ConcertSync/ConcertSyncServer/Source/ConcertSyncServer/Private/Replication/Muting/MuteManager.cpp
@@ -0,0 +1,503 @@
+﻿// Copyright Epic Games, Inc. All Rights Reserved.
+
+#include "MuteManager.h"
+
+#include "ConcertLogGlobal.h"
+#include "IConcertSession.h"
+#include "Misc/ObjectPathOuterIterator.h"
+#include "Replication/ChangeStreamSharedUtils.h"
+#include "Replication/Messages/Muting.h"
+#include "Replication/Misc/ReplicatedObjectHierarchyCache.h"
+#include "Replication/Muting/ObjectHierarchyAdapter.h"
+#include "Replication/Util/LogUtils.h"
+
+#include "HAL/IConsoleManager.h"
+#include "Misc/ScopeExit.h"
+
+namespace UE::ConcertSyncServer::Replication
+{
+	TAutoConsoleVariable<bool> CVarLogMuteRequestsAndResponsesOnServer(
+		TEXT("Concert.Replication.LogMuteRequestsAndResponsesOnServer"),
+		false,
+		TEXT("Whether to log changes to the mute state.")
+		);
+	
+	FMuteManager::FMuteManager(
+		IConcertSession& InSession,
+		const ConcertSyncCore::FReplicatedObjectHierarchyCache& ServerObjectCache,
+		EConcertSyncSessionFlags InSessionFlags
+		)
+		: Session(InSession)
+		, ServerObjectCache(ServerObjectCache)
+		, SessionFlags(InSessionFlags)
+	{
+		Session.RegisterCustomRequestHandler<FConcertReplication_ChangeMuteState_Request, FConcertReplication_ChangeMuteState_Response>(this, &FMuteManager::HandleChangeMuteStateRequest);
+		Session.RegisterCustomRequestHandler<FConcertReplication_QueryMuteState_Request, FConcertReplication_QueryMuteState_Response>(this, &FMuteManager::HandleQueryMuteStateRequest);
+	}
+
+	FMuteManager::~FMuteManager()
+	{
+		Session.UnregisterCustomRequestHandler<FConcertReplication_ChangeMuteState_Request>();
+		Session.UnregisterCustomRequestHandler<FConcertReplication_QueryMuteState_Request>();
+	}
+
+	bool FMuteManager::IsMuted(const FSoftObjectPath& Object) const
+	{
+		const FMuteData* ObjectMuteState = MuteState.Find(Object);
+		return ObjectMuteState
+			&& (ObjectMuteState->State == EMuteState::ExplicitlyMuted || ObjectMuteState->State == EMuteState::ImplicitlyMuted);
+	}
+
+	TOptional<FMuteManager::EMuteState> FMuteManager::GetMuteState(const FSoftObjectPath& Object) const
+	{
+		const FMuteData* ObjectMuteState = MuteState.Find(Object);
+		return ObjectMuteState
+			? ObjectMuteState->State
+			: TOptional<EMuteState>{};
+	}
+
+	TOptional<FConcertReplication_ObjectMuteSetting> FMuteManager::GetExplicitMuteSetting(const FSoftObjectPath& Object) const
+	{
+		const FMuteData* ObjectMuteState = MuteState.Find(Object);
+		return ObjectMuteState
+			? ObjectMuteState->MuteSetting
+			: TOptional<FConcertReplication_ObjectMuteSetting>{};
+	}
+
+	bool FMuteManager::ValidateRequest(
+		const FConcertReplication_ChangeMuteState_Request& Request,
+		const IMuteValidationObjectHierarchy* OverrideServerObjectCache,
+		TFunctionRef<void(const FSoftObjectPath& ObjectPath)> OnRejection
+		) const
+	{
+		const bool bFeatureSupported = EnumHasAnyFlags(SessionFlags, EConcertSyncSessionFlags::ShouldAllowGlobalMuting);
+		if (!ensureMsgf(bFeatureSupported, TEXT("Don't call this if the feature is not enabled!")))
+		{
+			return false;
+		}
+
+		return OverrideServerObjectCache
+			? ValidateRequestInternal(Request, *OverrideServerObjectCache, OnRejection)
+			: ValidateRequestInternal(Request, FObjectHierarchyAdapter(ServerObjectCache), OnRejection);
+	}
+
+	bool FMuteManager::ApplyRequestAndEnumerateSyncControl(
+		const FConcertReplication_ChangeMuteState_Request& Request,
+		const FOnSyncControlChange& OnSyncControlChange
+		)
+	{
+		const bool bIsValid = ValidateRequest(Request);
+		if (ensureAlwaysMsgf(bIsValid, TEXT("This is a double-check and it failed. The caller should have validate the request beforehand.")))
+		{
+			ApplyRequestInternal(Request);
+
+			check(OnRefreshSyncControlButSkipSendingToClientsDelegate.IsBound());
+			OnRefreshSyncControlButSkipSendingToClientsDelegate.Execute(OnSyncControlChange);
+		}
+		
+		return bIsValid;
+	}
+
+	FConcertReplication_ChangeSyncControl FMuteManager::ApplyManualRequest(const FGuid& EndpointId, const FConcertReplication_ChangeMuteState_Request& Request)
+	{
+		ApplyRequestInternal(Request);
+		return OnRefreshSyncControlAndSendToAllClientsExceptDelegate.Execute(EndpointId);
+	}
+
+	void FMuteManager::PostApplyStreamChange(const FGuid& ClientId, TConstArrayView<FConcertObjectInStreamID> AddedObjects, TConstArrayView<FConcertObjectInStreamID> RemovedObjects)
+	{
+		// If Request adds new subobjects that should be muted implicitly by existing rules, update MuteState.
+		TrackAddedSubobjectsForImplicitMuting(AddedObjects);
+		// Objects that are no longer referenced by any clients must be removed (to avoid leaks)
+		UnmuteObjectsIfUnreferenced(RemovedObjects);
+		
+		// This updates sync control, which may remove sync control from other clients.
+		check(OnMuteStateChangedDelegate.IsBound());
+		OnMuteStateChangedDelegate.Execute(ClientId);
+	}
+
+	void FMuteManager::OnPostClientLeft(const TArray<FConcertReplicationStream>& ClientStreams)
+	{
+		// We'll treat this case as a FConcertReplication_ChangeStream_Request that removes all objects 
+		TSet<FSoftObjectPath> VisitedObjects;
+		for (const FConcertReplicationStream& Stream : ClientStreams)
+		{
+			for (const TPair<FSoftObjectPath, FConcertReplicatedObjectInfo>& ObjectPair : Stream.BaseDescription.ReplicationMap.ReplicatedObjects)
+			{
+				const FSoftObjectPath ObjectPath = ObjectPair.Key; 
+				if (!VisitedObjects.Contains(ObjectPath))
+				{
+					VisitedObjects.Add(ObjectPath);
+					UnmuteObjectIfUnreferenced(ObjectPath);
+				}
+			}
+		}
+	}
+
+	void FMuteManager::TrackAddedSubobjectsForImplicitMuting(TConstArrayView<FConcertObjectInStreamID> AddedObjects)
+	{
+		for (const FConcertObjectInStreamID& Object : AddedObjects)
+		{
+			const FSoftObjectPath ObjectPath = Object.Object;
+			if (MuteState.Contains(ObjectPath))
+			{
+				continue;
+			}
+
+			const auto[_, ParentMuteData] = FindAffectingParentState(ObjectPath);
+			if (ParentMuteData && ParentMuteData->State == EMuteState::ExplicitlyUnmuted)
+			{
+				MuteState.Add(ObjectPath).State = EMuteState::ImplicitlyMuted;
+			}
+			else if (ParentMuteData && ParentMuteData->State == EMuteState::ExplicitlyMuted)
+			{
+				MuteState.Add(ObjectPath).State = EMuteState::ImplicitlyUnmuted;
+			}
+		}
+	}
+
+	void FMuteManager::UnmuteObjectsIfUnreferenced(TConstArrayView<FConcertObjectInStreamID> RemovedObjects)
+	{
+		// We'll avoid processing any object again that we've already removed... not doing so would unnecessarily iterate the hierarchy.
+		TSet<FSoftObjectPath> ProcessedObjects;
+		
+		for (const FConcertObjectInStreamID& RemovedObject : RemovedObjects)
+		{
+			const FSoftObjectPath& RemovedObjectPath = RemovedObject.Object;
+			if (!ProcessedObjects.Contains(RemovedObjectPath))
+			{
+				UnmuteObjectIfUnreferenced(RemovedObjectPath, [&ProcessedObjects](const FSoftObjectPath& ObjectPath)
+				{
+					ProcessedObjects.Add(ObjectPath);
+				});
+			}
+		}
+	}
+
+	template<CObjectProcessable TCallback>
+	void FMuteManager::UnmuteObjectIfUnreferenced(const FSoftObjectPath& RemovedObject, TCallback&& OnRemoved)
+	{
+		// If RemovedObject is referenced or has known subobjects, the mute state should tick around.
+		const bool bHasChildrenOrIsReferenced = ServerObjectCache.HasChildren(RemovedObject);
+		if (bHasChildrenOrIsReferenced)
+		{
+			return;
+		}
+		
+		const auto [AffectingParent, ParentMuteData] = FindAffectingParentState(RemovedObject);
+		// RemovedObject is not referenced by any streams, so it has no business being muted.
+		RemoveMuteState(RemovedObject, OnRemoved);
+
+		// There may be a parent object, which also no longer has known children ... so remove that as well.
+		if (ParentMuteData)
+		{
+			UnmuteObjectIfUnreferenced(AffectingParent, OnRemoved);
+		}
+	}
+
+	EConcertSessionResponseCode FMuteManager::HandleQueryMuteStateRequest(
+		const FConcertSessionContext& Context,
+		const FConcertReplication_QueryMuteState_Request& Request,
+		FConcertReplication_QueryMuteState_Response& Response
+		)
+	{
+		if (!EnumHasAnyFlags(SessionFlags, EConcertSyncSessionFlags::ShouldAllowGlobalMuting))
+		{
+			return EConcertSessionResponseCode::Failed;
+		}
+
+		Response.ErrorCode = EReplicationResponseErrorCode::Handled;
+		const auto AddObjectToResponse = [this, &Response](const FSoftObjectPath& ObjectPath, const FMuteData& ObjectMuteData)
+		{
+			switch (ObjectMuteData.State)
+			{
+			case EMuteState::ExplicitlyMuted:
+				Response.ExplicitlyMutedObjects.Add(ObjectPath, *ObjectMuteData.MuteSetting);
+				break;
+			case EMuteState::ExplicitlyUnmuted:
+				Response.ExplicitlyUnmutedObjects.Add(ObjectPath, *ObjectMuteData.MuteSetting);
+				break;
+			case EMuteState::ImplicitlyMuted:
+				Response.ImplicitlyMutedObjects.Add(ObjectPath);
+				break;
+			case EMuteState::ImplicitlyUnmuted:
+				Response.ImplicitlyUnmutedObjects.Add(ObjectPath);
+				break;
+			}
+		};
+
+		if (Request.WantsAllObjects())
+		{
+			for (const TPair<FSoftObjectPath, FMuteData>& Pair : MuteState)
+			{
+				AddObjectToResponse(Pair.Key, Pair.Value);
+			}
+		}
+		else
+		{
+			for (const FSoftObjectPath& QueriedPath : Request.QueriedObjects)
+			{
+				if (const FMuteData* ObjectMuteData = MuteState.Find(QueriedPath))
+				{
+					AddObjectToResponse(QueriedPath, *ObjectMuteData);
+				}
+			}
+		}
+		
+		return EConcertSessionResponseCode::Success;
+	}
+
+	EConcertSessionResponseCode FMuteManager::HandleChangeMuteStateRequest(
+		const FConcertSessionContext& Context,
+		const FConcertReplication_ChangeMuteState_Request& Request,
+		FConcertReplication_ChangeMuteState_Response& Response
+		)
+	{
+		LogNetworkMessage(CVarLogMuteRequestsAndResponsesOnServer, Request, [this, &Context](){ return GetClientName(Session, Context.SourceEndpointId); });
+		ON_SCOPE_EXIT{ LogNetworkMessage(CVarLogMuteRequestsAndResponsesOnServer, Response, [this, &Context](){ return GetClientName(Session, Context.SourceEndpointId); }); };
+		
+		if (!EnumHasAnyFlags(SessionFlags, EConcertSyncSessionFlags::ShouldAllowGlobalMuting))
+		{
+			return EConcertSessionResponseCode::Failed;
+		}
+
+		if (!ValidateRequestInternal(Request, Response))
+		{
+			Response.ErrorCode = EConcertReplicationMuteErrorCode::Rejected;
+			return EConcertSessionResponseCode::Success;
+		}
+
+		Response.ErrorCode = EConcertReplicationMuteErrorCode::Accepted;
+		ApplyRequestInternal(Request);
+		
+		// Sync control will now 1. generate a sync control we can embed in the response, and 2. send a network event to all other clients.
+		// The embedded sync control will contain only those objects which changed sync control
+		// Fyi, the new sync control may contain more objects than were specified in the request since parent objects can also cause their subobjects' mute state to change.
+		check(OnRefreshSyncControlAndSendToAllClientsExceptDelegate.IsBound());
+		Response.SyncControl = OnRefreshSyncControlAndSendToAllClientsExceptDelegate.Execute(Context.SourceEndpointId);
+#if UE_BUILD_DEBUG // In debug builds so we don't tank performance
+		checkf(!Response.SyncControl.DoesAtLeastOneObjectLoseSyncControl(), TEXT("Clients are supposed to infer that objects lose sync control. We only tell them which objects gained sync control."));
+#endif
+
+		OnMuteRequestAppliedDelegate.Broadcast(Context.SourceEndpointId, Request);
+		return EConcertSessionResponseCode::Success;
+	}
+
+	bool FMuteManager::ValidateRequestInternal(const FConcertReplication_ChangeMuteState_Request& Request, FConcertReplication_ChangeMuteState_Response& Response) const
+	{
+		const auto Reject = [&Response](const FSoftObjectPath& ObjectPath){ Response.RejectionReasons.Add(ObjectPath); };
+		return ValidateRequestInternal(Request, FObjectHierarchyAdapter(ServerObjectCache), Reject);
+	}
+
+	bool FMuteManager::ValidateRequestInternal(
+		const FConcertReplication_ChangeMuteState_Request& Request, 
+		const IMuteValidationObjectHierarchy& ObjectCache,
+		TFunctionRef<void(const FSoftObjectPath& ObjectPath)> OnRejection
+		) const
+	{
+		bool bIsValidRequest = true;
+		const auto Reject = [&OnRejection, &bIsValidRequest](const FSoftObjectPath& ObjectPath){ OnRejection(ObjectPath); bIsValidRequest = false; };
+		
+		const auto SharedValidateSetting = [this, &Reject, &ObjectCache](const FSoftObjectPath& Object, const FConcertReplication_ObjectMuteSetting& Setting)
+		{
+			const bool bIsObjectReferenced = ObjectCache.IsObjectReferencedDirectly(Object);
+			const bool bAppliesToSubobjects = ConcertSyncCore::AffectSubobjects(Setting.Flags);
+			if (!bIsObjectReferenced && !bAppliesToSubobjects)
+			{
+				Reject(Object);
+				return;
+			}
+
+			const bool bHasSubobjects = ObjectCache.HasChildren(Object);
+			if (!bIsObjectReferenced && bAppliesToSubobjects && !bHasSubobjects)
+			{
+				Reject(Object);
+			}
+		};
+		
+		for (const TPair<FSoftObjectPath, FConcertReplication_ObjectMuteSetting>& ToMute : Request.ObjectsToMute)
+		{
+			SharedValidateSetting(ToMute.Key, ToMute.Value);
+
+			// Cannot mute & unmute at the same time
+			if (Request.ObjectsToUnmute.Contains(ToMute.Key))
+			{
+				Reject(ToMute.Key);
+			}
+		}
+		for (const TPair<FSoftObjectPath, FConcertReplication_ObjectMuteSetting>& ToUnmute : Request.ObjectsToUnmute)
+		{
+			// If something is muted, the user can always unmute it. If it is already unmuted, it's still valid (but a no-op).
+			if (!MuteState.Contains(ToUnmute.Key))
+			{
+				SharedValidateSetting(ToUnmute.Key, ToUnmute.Value);
+			}
+		}
+		
+		return bIsValidRequest;
+	}
+
+	void FMuteManager::ApplyRequestInternal(const FConcertReplication_ChangeMuteState_Request& Request)
+	{
+		UE_LOG(LogConcert, Log, TEXT("Explicitly muting %d and unmuting %d objects."), Request.ObjectsToMute.Num(), Request.ObjectsToUnmute.Num());
+		
+		for (const TPair<FSoftObjectPath, FConcertReplication_ObjectMuteSetting>& ToMute : Request.ObjectsToMute)
+		{
+			const FSoftObjectPath& ObjectToMute = ToMute.Key;
+			const FConcertReplication_ObjectMuteSetting& MuteSetting = ToMute.Value;
+
+			// Muting object again with same parameters is a valid: not work to do if so.
+			const bool bAffectsChildren = ConcertSyncCore::AffectSubobjects(MuteSetting.Flags);
+			{
+				const FMuteData* ObjectMuteState = MuteState.Find(ObjectToMute);
+				const bool bAlreadyAtSetting = ObjectMuteState && ObjectMuteState->MuteSetting && ObjectMuteState->MuteSetting == ToMute.Value;
+				if (bAlreadyAtSetting)
+				{
+					continue;
+				}
+
+				// If the request is changing the flag from IncludeSubobjects -> None, the children need their implicit state updated.
+				const bool bPreviousAffectedChildren = ObjectMuteState && ObjectMuteState->AffectsSubobjects(); 
+				if (!bAffectsChildren && bPreviousAffectedChildren)
+				{
+					RemoveMuteState(ObjectToMute);
+				}
+				// ObjectMuteState may now be invalid due to possible removal
+			}
+
+			const FMuteData NewMuteData
+			{
+				MuteSetting,
+				EMuteState::ExplicitlyMuted
+			};
+			MuteState.Add(ObjectToMute, NewMuteData);
+
+			// Children may need their implicit state updated
+			if (bAffectsChildren)
+			{
+				UpdateImplicitStateUnder(ObjectToMute, EMuteState::ImplicitlyMuted);
+			}
+		}
+
+		for (const TPair<FSoftObjectPath, FConcertReplication_ObjectMuteSetting>& ToMute : Request.ObjectsToUnmute)
+		{
+			const FSoftObjectPath& ObjectToUnmute = ToMute.Key;
+			const FConcertReplication_ObjectMuteSetting& MuteSetting = ToMute.Value;
+			
+			const auto[AffectingParent, ParentMuteState] = FindAffectingParentState(ObjectToUnmute);
+			// ParentMuteState can either be:
+			//	- EMuteState::ExplicitlyMuted: all children should be muted
+			//	- EMuteState::ExplicitlyUnmuted: by invariant, there must be a parent that is ExplicitlyMuted with EConcertReplicationMuteFlags::ObjectAndSubobjects.
+			// Hence in both cases, children need to be updated to be marked as implicitly muted.
+			if (ParentMuteState)
+			{
+				MuteState.Add(ObjectToUnmute, { MuteSetting, EMuteState::ExplicitlyUnmuted });
+				if (ConcertSyncCore::AffectSubobjects(MuteSetting.Flags))
+				{
+					UpdateImplicitStateUnder(ObjectToUnmute, EMuteState::ImplicitlyUnmuted);
+				}
+			}
+			else
+			{
+				// If no parent state affects ObjectToUnmute, the entire subtree must be removed.
+				FMuteData MuteData;
+				const bool bHadData = MuteState.RemoveAndCopyValue(ObjectToUnmute, MuteData);
+				if (bHadData && MuteData.AffectsSubobjects())
+				{
+					ClearAllChildStateUnder(ObjectToUnmute);
+				}
+			}
+		}
+	}
+
+	template<CObjectProcessable TCallback>
+	void FMuteManager::RemoveMuteState(const FSoftObjectPath& Object, TCallback&& OnRemoved)
+	{
+		const FMuteData* ObjectMuteData = MuteState.Find(Object);
+		if (!ObjectMuteData)
+		{
+			return;
+		}
+		
+		// If an Object affects its subobjects...
+		if (ObjectMuteData->AffectsSubobjects())
+		{
+			const auto[AffectingParent, ParentMuteData] = FindAffectingParentState(Object);
+			if (ParentMuteData)
+			{
+				// ... and has a parent also affecting subobjects, then all of Object's children are now affected by that parent.
+				const EMuteState ImplicitNewChildState = ParentMuteData->State == EMuteState::ExplicitlyMuted ? EMuteState::ImplicitlyMuted : EMuteState::ImplicitlyUnmuted;
+				UpdateImplicitStateUnder(Object, ImplicitNewChildState);
+			}
+			else
+			{
+				// ... and does not have any parent affecting subobjects, then clear all child state
+				ClearAllChildStateUnder(Object, OnRemoved);
+			}
+		}
+
+		MuteState.Remove(Object);
+		OnRemoved(Object);
+	}
+
+	void FMuteManager::UpdateImplicitStateUnder(const FSoftObjectPath& Parent, EMuteState NewImplicitState)
+	{
+		check(NewImplicitState == EMuteState::ImplicitlyMuted || NewImplicitState == EMuteState::ImplicitlyUnmuted);
+		
+		ServerObjectCache.TraverseTopToBottom([this, NewImplicitState](const ConcertSyncCore::FChildRelation& Relation)
+		{
+			FMuteData* SubobjectMuteState = MuteState.Find(Relation.Child.Object);
+			if (!SubobjectMuteState)
+			{
+				MuteState.Add(Relation.Child.Object, { {}, NewImplicitState });
+			}
+			else if (SubobjectMuteState->MuteSetting
+				&& ConcertSyncCore::AffectSubobjects(SubobjectMuteState->MuteSetting->Flags))
+			{
+				return ConcertSyncCore::ETreeTraversalBehavior::SkipSubtree;
+			}
+			else
+			{
+				SubobjectMuteState->State = NewImplicitState;
+			}
+					
+			return ConcertSyncCore::ETreeTraversalBehavior::Continue;
+		}, Parent);
+	}
+	
+	template<CObjectProcessable TCallback>
+	void FMuteManager::ClearAllChildStateUnder(const FSoftObjectPath& ParentObject, TCallback&& OnRemoved)
+	{
+		ServerObjectCache.TraverseTopToBottom([this, &OnRemoved](const ConcertSyncCore::FChildRelation& Relation)
+		{
+			FMuteData* SubobjectMuteState = MuteState.Find(Relation.Child.Object);
+			if (!SubobjectMuteState)
+			{
+				return ConcertSyncCore::ETreeTraversalBehavior::Continue;
+			}
+
+			MuteState.Remove(Relation.Child.Object);
+			OnRemoved(Relation.Child.Object);
+			return ConcertSyncCore::ETreeTraversalBehavior::Continue;
+		}, ParentObject);
+	}
+	
+	TPair<FSoftObjectPath, const FMuteManager::FMuteData*> FMuteManager::FindAffectingParentState(const FSoftObjectPath& Subobject) const
+	{
+		for (ConcertSyncCore::FObjectPathOuterIterator It(Subobject); It; ++It)
+		{
+			const FSoftObjectPath& Parent = *It;
+			if (const FMuteData* ParentMuteState = MuteState.Find(Parent)
+				; ParentMuteState
+				// Does it have the IncludeSubobjects flag
+				&& ParentMuteState->MuteSetting && ConcertSyncCore::AffectSubobjects(ParentMuteState->MuteSetting->Flags)
+				// Only explicit objects should have the IncludeSubobjects flag
+				&& ensure(ParentMuteState->IsExplicit()))
+			{
+				return { Parent, ParentMuteState };
+			}
+		}
+		
+		return {};
+	}
+}
diff --git a/Engine/Plugins/Developer/Concert/ConcertSync/ConcertSyncServer/Source/ConcertSyncServer/Private/Replication/Muting/MuteManager.h b/Engine/Plugins/Developer/Concert/ConcertSync/ConcertSyncServer/Source/ConcertSyncServer/Private/Replication/Muting/MuteManager.h
new file mode 100644
--- /dev/null
+++ b/Engine/Plugins/Developer/Concert/ConcertSync/ConcertSyncServer/Source/ConcertSyncServer/Private/Replication/Muting/MuteManager.h
@@ -0,0 +1,212 @@
+﻿// Copyright Epic Games, Inc. All Rights Reserved.
+
+#pragma once
+
+#include "ConcertMessages.h"
+#include "ConcertSyncSessionFlags.h"
+#include "IConcertSessionHandler.h"
+#include "Replication/Messages/Muting.h"
+
+#include "Misc/Optional.h"
+#include "Templates/Function.h"
+#include "Templates/UnrealTemplate.h"
+
+#include <type_traits>
+
+class IConcertSession;
+struct FConcertReplication_ChangeStream_Request;
+struct FConcertReplicationStream;
+struct FConcertObjectInStreamID;
+
+namespace UE::ConcertSyncCore { class FReplicatedObjectHierarchyCache; }
+namespace UE::ConcertSyncServer::Replication {
+	class IMuteValidationObjectHierarchy;
+	class IRegistrationEnumerator; }
+
+namespace UE::ConcertSyncServer::Replication
+{
+	/** Invocable with signature void([](const FConcertObjectInStreamID&){}) which calls the passed in lambda */
+	template<typename TEnumerateObjects>
+	concept CObjectEnumeratable = requires(TEnumerateObjects EnumerateObjects)
+	{
+		EnumerateObjects([](const FConcertObjectInStreamID&){});
+	};
+	
+	/** Invocable with signature void(FSoftObjectPath) */
+	template<typename TLambda>
+	concept CObjectProcessable = std::is_invocable_v<TLambda, const FSoftObjectPath&>;
+	
+	/**
+	 * Manages the session's replication mute state.
+	 * Handles FConcertReplication_ChangeMuteState_Request and FConcertReplication_QueryMuteState_Request requets.
+	 */
+	class FMuteManager : public FNoncopyable
+	{
+	public:
+
+		using FOnSyncControlChange = TFunctionRef<void(const FGuid& ClientId, FConcertReplication_ChangeSyncControl&& SyncControlChange)>;
+
+		DECLARE_DELEGATE_OneParam(FOnMuteStateChangedByClient, const FGuid& ClientId);
+		DECLARE_DELEGATE_RetVal_OneParam(FConcertReplication_ChangeSyncControl, FRefreshSyncControlForMuteChange, const FGuid& ClientId);
+		DECLARE_DELEGATE_OneParam(FRefreshSyncControlForClients, const FOnSyncControlChange& OnSyncControlChange);
+		
+		/** Notifies manager about applied event. Generates mute activity. */
+		DECLARE_MULTICAST_DELEGATE_TwoParams(FOnMuteRequestApplied, const FGuid& ClientId, const FConcertReplication_ChangeMuteState_Request& Request);
+		
+		enum class EMuteState : uint8
+		{
+			/** Object is muted */
+			ExplicitlyMuted,
+			/** A parent object is ExplicitlyMuted and has EConcertReplicationMuteFlags::ObjectAndSubobjects set but is object and its children are supposed to be implicitly unmuted.  */
+			ExplicitlyUnmuted,
+			
+			/** Object is muted because one of its parent objects is ExplicitlyMuted and EConcertReplicationMuteFlags::ObjectAndSubobjects set. */
+			ImplicitlyMuted,
+			/** Object is unmuted because one of its parent objects is ExplicitlyUnmuted and has EConcertReplicationMuteFlags::ObjectAndSubobjects set. */
+			ImplicitlyUnmuted,
+		};
+
+		FMuteManager(
+			IConcertSession& InSession UE_LIFETIMEBOUND,
+			const ConcertSyncCore::FReplicatedObjectHierarchyCache& ServerObjectCache UE_LIFETIMEBOUND,
+			EConcertSyncSessionFlags InSessionFlags
+			);
+		~FMuteManager();
+
+		/** @return Whether the object is globally muted. */
+		bool IsMuted(const FSoftObjectPath& Object) const;
+
+		/** @return The mute state of Object. If it was not explicitly muted and no parent object affects this subobject, then the optional is unset. */
+		TOptional<EMuteState> GetMuteState(const FSoftObjectPath& Object) const;
+		/** @return The mute setting of Object, if it is explicitly set (i.e. not affected by a parent object). */
+		TOptional<FConcertReplication_ObjectMuteSetting> GetExplicitMuteSetting(const FSoftObjectPath& Object) const;
+		
+		/**
+		 * Checks whether Request would be valid to apply, optionally providing a future server state.
+		 * @param Request The request to validate
+		 * @param OverrideServerObjectCache The server cache to use. Defaults to the internal one if left nullptr. Useful for passing in a future server state.
+		 * @param OnRejection Called for every rejected object.
+		 */
+		bool ValidateRequest(
+			const FConcertReplication_ChangeMuteState_Request& Request,
+			const IMuteValidationObjectHierarchy* OverrideServerObjectCache = nullptr,
+			TFunctionRef<void(const FSoftObjectPath& ObjectPath)> OnRejection = [](const FSoftObjectPath&){}
+			) const;
+		
+		/**
+		 * This validates and then applies Request.
+		 * This version does not send any sync control updates to other clients. You must update them yourself.
+		 * 
+		 * @param Request The request to apply
+		 * @param OnSyncControlChange Called for every sync control change that is made.
+		 * @return Whether the request was applied.
+		 */
+		bool ApplyRequestAndEnumerateSyncControl(
+			const FConcertReplication_ChangeMuteState_Request& Request,
+			const FOnSyncControlChange& OnSyncControlChange
+			);
+		/**
+		 * Applies Request as if it was sent by EndpointId.
+		 * @return The sync control that the client identified by EndpointId has gained. Does NOT include the sync control they lost.
+		 */
+		FConcertReplication_ChangeSyncControl ApplyManualRequest(const FGuid& EndpointId, const FConcertReplication_ChangeMuteState_Request& Request);
+		
+		/** Called right after objects have been unregistered from ClientId's streams. */
+		void PostApplyStreamChange(const FGuid& ClientId, TConstArrayView<FConcertObjectInStreamID> AddedObjects, TConstArrayView<FConcertObjectInStreamID> RemovedObjects);
+		/** Called by FConcertServerReplicationManager when client leaves replication. */
+		void OnPostClientLeft(const TArray<FConcertReplicationStream>& ClientStreams);
+
+		/** Updates sync control if needed in response to a client changing their mute state indirectly (e.g. because client removed object from their stream). */
+		FOnMuteStateChangedByClient& OnRefreshSyncControlForIndirectMuteChange() { return OnMuteStateChangedDelegate; }
+		/** Updates sync control for all clients, sends an update to all clients but ClientId, and returns the sync control to embed into the mute response. */
+		FRefreshSyncControlForMuteChange& OnRefreshSyncControlAndSendToAllClientsExcept() { return OnRefreshSyncControlAndSendToAllClientsExceptDelegate; }
+		/** Updates sync control for all clients but does not send any update to any clients. The passed in callback receives all sync control changes for the clients. */
+		FRefreshSyncControlForClients& OnRefreshSyncControlButSkipSendingToClients() { return OnRefreshSyncControlButSkipSendingToClientsDelegate; }
+		
+		FOnMuteRequestApplied& OnMuteRequestApplied() { return OnMuteRequestAppliedDelegate; }
+		
+	private:
+
+		/** The session to receive requests on. */
+		IConcertSession& Session;
+		/** Keeps track of all objects that are in client streams. Allows efficient traversing of subobject hierarchy. */
+		const ConcertSyncCore::FReplicatedObjectHierarchyCache& ServerObjectCache;
+
+		/** If ShouldAllowGlobalMuting is not set, the requests are not served. */
+		const EConcertSyncSessionFlags SessionFlags;
+
+		/** Broadcasts after the mute state has been changed. */
+		FOnMuteStateChangedByClient OnMuteStateChangedDelegate;
+		/**
+		 * Delegate into FSyncControlManager
+		 * 
+		 * Updates sync control for all clients, sends an update to all clients but ClientId, and returns the sync control to embed into the mute response.
+		 * The sync control will only contain RestrictToObjects.
+		 */
+		FRefreshSyncControlForMuteChange OnRefreshSyncControlAndSendToAllClientsExceptDelegate;
+		/** Delegate into FSyncControlManager.  */
+		FRefreshSyncControlForClients OnRefreshSyncControlButSkipSendingToClientsDelegate;
+		
+		/** Broadcasts after a mute request has been applied. */
+		FOnMuteRequestApplied OnMuteRequestAppliedDelegate;
+		
+		struct FMuteData
+		{
+			/** Only set if the object was explicitly muted or unmuted. */
+			TOptional<FConcertReplication_ObjectMuteSetting> MuteSetting;
+			
+			EMuteState State;
+
+			bool IsExplicit() const { return MuteSetting.IsSet(); }
+			bool IsImplicit() const { return !IsExplicit(); }
+			bool AffectsSubobjects() const { return MuteSetting && ConcertSyncCore::AffectSubobjects(MuteSetting->Flags); }
+		};
+		/**
+		 * Holds the mute state of all objects.
+		 * 
+		 * Objects only show up in this map if they are either explicitly or implicitly affected by a mute setting:
+		 * - explicit means that a request listed them to be muted / unmuted
+		 * - implicit means that a parent object (i.e. an outer) is explicitly muted or unmuted.
+		 * Hence, if an object is not in this map it means that neither it nor any of its parents are explicitly/implicitly muted nor unmuted.
+		 */
+		TMap<FSoftObjectPath, FMuteData> MuteState;
+		
+		/** For each added object, checks whether it is a subobject of an object that is muted with the EConcertReplicationMuteFlags::ObjectAndSubobjects setting. */
+		void TrackAddedSubobjectsForImplicitMuting(TConstArrayView<FConcertObjectInStreamID> AddedObjects);
+		/** Calls RemoveMuteState on every object that is unreferenced by client streams (directly or indirectly via a subobject). */
+		void UnmuteObjectsIfUnreferenced(TConstArrayView<FConcertObjectInStreamID> RemovedObjects);
+		/** Calls RemoveMuteState if RemovedObject is unreferenced by client streams (directly or indirectly via a subobject). */
+		template<CObjectProcessable TCallback>
+		void UnmuteObjectIfUnreferenced(const FSoftObjectPath& RemovedObject, TCallback&& OnRemoved);
+		void UnmuteObjectIfUnreferenced(const FSoftObjectPath& RemovedObject) { UnmuteObjectIfUnreferenced(RemovedObject, [](const FSoftObjectPath& ObjectPath){}); }
+
+		/** Handles requests to query mute states. */
+		EConcertSessionResponseCode HandleQueryMuteStateRequest(const FConcertSessionContext& Context, const FConcertReplication_QueryMuteState_Request& Request, FConcertReplication_QueryMuteState_Response& Response);
+		/** Handles requests to change the mute states */
+		EConcertSessionResponseCode HandleChangeMuteStateRequest(const FConcertSessionContext& Context, const FConcertReplication_ChangeMuteState_Request& Request, FConcertReplication_ChangeMuteState_Response& Response);
+
+		/** Checks whether Request is valid to apply. */
+		bool ValidateRequestInternal(const FConcertReplication_ChangeMuteState_Request& Request, FConcertReplication_ChangeMuteState_Response& Response) const;
+		/** Checks whether Request is valid to apply. */
+		bool ValidateRequestInternal(const FConcertReplication_ChangeMuteState_Request& Request, const IMuteValidationObjectHierarchy& ObjectCache, TFunctionRef<void(const FSoftObjectPath& ObjectPath)> OnRejection) const;
+		/** Updates the internal state (assuming it is a valid request - call ValidateRequest before). */
+		void ApplyRequestInternal(const FConcertReplication_ChangeMuteState_Request& Request);
+
+		/** Removes all mute state for Object and transitively updates subobjects if applicable. */
+		template<CObjectProcessable TCallback>
+		void RemoveMuteState(const FSoftObjectPath& Object, TCallback&& OnRemoved);
+		void RemoveMuteState(const FSoftObjectPath& Object) { RemoveMuteState(Object, [](const FSoftObjectPath&){}); } 
+		
+		/** Sets the mute state of all subobjects of Parent until reaching a subobject that also has EConcertReplicationMuteFlags::ObjectAndSubobjects set. */
+		void UpdateImplicitStateUnder(const FSoftObjectPath& Parent, EMuteState NewImplicitState);
+		
+		/** Removes all subobject mute state affected by Parent's mute state assuming that ParentObject has no parents with the EConcertReplicationMuteFlags::ObjectAndSubobjects. */
+		template<CObjectProcessable TCallback>
+		void ClearAllChildStateUnder(const FSoftObjectPath& ParentObject, TCallback&& OnRemoved);
+		void ClearAllChildStateUnder(const FSoftObjectPath& ParentObject) { ClearAllChildStateUnder(ParentObject, [](const FSoftObjectPath&){}); }
+		
+		/** Walks up the object hierarchy and returns the parent that affects ObjectPath, if any. */
+		TPair<FSoftObjectPath, const FMuteData*> FindAffectingParentState(const FSoftObjectPath& Subobject) const;
+	};
+}
+
diff --git a/Engine/Plugins/Developer/Concert/ConcertSync/ConcertSyncServer/Source/ConcertSyncServer/Private/Replication/Muting/PredictedStateObjectHierarchy.cpp b/Engine/Plugins/Developer/Concert/ConcertSync/ConcertSyncServer/Source/ConcertSyncServer/Private/Replication/Muting/PredictedStateObjectHierarchy.cpp
--- a/Engine/Plugins/Developer/Concert/ConcertSync/ConcertSyncServer/Source/ConcertSyncServer/Private/Replication/Muting/PredictedStateObjectHierarchy.cpp
+++ b/Engine/Plugins/Developer/Concert/ConcertSync/ConcertSyncServer/Source/ConcertSyncServer/Private/Replication/Muting/PredictedStateObjectHierarchy.cpp
@@ -7,11 +7,11 @@
 
 namespace UE::ConcertSyncServer::Replication
 {
-	void FPredictedStateObjectHierarchy::AddClients(const TMap<FGuid, TArray<FConcertReplicationStream>>& Clients)
+	void FPredictedStateObjectHierarchy::AddClients(const TMap<FGuid, FConcertReplicationStreamArray>& Clients)
 	{
-		for (const TPair<FGuid, TArray<FConcertReplicationStream>>& Pair : Clients)
+		for (const TPair<FGuid, FConcertReplicationStreamArray>& Pair : Clients)
 		{
-			AddClientData(Pair.Key, Pair.Value);
+			AddClientData(Pair.Key, Pair.Value.Streams);
 		}
 	}
 
diff --git a/Engine/Plugins/Developer/Concert/ConcertSync/ConcertSyncServer/Source/ConcertSyncServer/Private/Replication/Muting/PredictedStateObjectHierarchy.h b/Engine/Plugins/Developer/Concert/ConcertSync/ConcertSyncServer/Source/ConcertSyncServer/Private/Replication/Muting/PredictedStateObjectHierarchy.h
--- a/Engine/Plugins/Developer/Concert/ConcertSync/ConcertSyncServer/Source/ConcertSyncServer/Private/Replication/Muting/PredictedStateObjectHierarchy.h
+++ b/Engine/Plugins/Developer/Concert/ConcertSync/ConcertSyncServer/Source/ConcertSyncServer/Private/Replication/Muting/PredictedStateObjectHierarchy.h
@@ -5,9 +5,11 @@
 #include "IMuteValidationObjectHierarchy.h"
 #include "Misc/ObjectPathHierarchy.h"
 #include "Replication/ConcertReplicationClient.h"
+#include "Replication/Data/ReplicationStreamArray.h"
 
 #include <type_traits>
 
+
 struct FConcertReplicationStream;
 
 namespace UE::ConcertSyncServer::Replication
@@ -23,7 +25,7 @@ namespace UE::ConcertSyncServer::Replication
 	public:
 
 		/** Adds clients from a map binding client id to stream content. */
-		void AddClients(const TMap<FGuid, TArray<FConcertReplicationStream>>& Clients);
+		void AddClients(const TMap<FGuid, FConcertReplicationStreamArray>& Clients);
 		/** Adds only those entries that pass ShouldIncludeFilter returns true on. */
 		template<typename TFilterLambda> requires std::is_invocable_r_v<bool, TFilterLambda, const FGuid&>
 		void AddClients(const TMap<FGuid, TUniquePtr<FConcertReplicationClient>>& Clients, TFilterLambda&& ShouldIncludeFilter);
diff --git a/Engine/Plugins/Developer/Concert/ConcertSync/ConcertSyncServer/Source/ConcertSyncServer/Private/Replication/SyncControlManager.cpp b/Engine/Plugins/Developer/Concert/ConcertSync/ConcertSyncServer/Source/ConcertSyncServer/Private/Replication/SyncControlManager.cpp
--- a/Engine/Plugins/Developer/Concert/ConcertSync/ConcertSyncServer/Source/ConcertSyncServer/Private/Replication/SyncControlManager.cpp
+++ b/Engine/Plugins/Developer/Concert/ConcertSync/ConcertSyncServer/Source/ConcertSyncServer/Private/Replication/SyncControlManager.cpp
@@ -3,8 +3,8 @@
 #include "SyncControlManager.h"
 
 #include "AuthorityManager.h"
-#include "MuteManager.h"
 #include "ConcertServer/Private/ConcertServerSession.h"
+#include "Muting/MuteManager.h"
 #include "Replication/Data/ReplicationStream.h"
 #include "Replication/Messages/SyncControl.h"
 
@@ -55,6 +55,12 @@ namespace UE::ConcertSyncServer::Replication
 		return ClientData && ClientData->ObjectsWithSyncControl.Contains(Object);
 	}
 
+	const TSet<FConcertObjectInStreamID>* FSyncControlManager::GetClientControlledObjects(const FGuid& ClientId) const
+	{
+		const FClientData* ClientData = PerClientData.Find(ClientId);
+		return ClientData ? &ClientData->ObjectsWithSyncControl : nullptr;
+	}
+
 	FConcertReplication_ChangeSyncControl FSyncControlManager::OnGenerateSyncControlForClientJoin(const FGuid& ClientId)
 	{
 		RefreshAndSendToAllClientsExcept(ClientId);
diff --git a/Engine/Plugins/Developer/Concert/ConcertSync/ConcertSyncServer/Source/ConcertSyncServer/Private/Replication/SyncControlManager.h b/Engine/Plugins/Developer/Concert/ConcertSync/ConcertSyncServer/Source/ConcertSyncServer/Private/Replication/SyncControlManager.h
--- a/Engine/Plugins/Developer/Concert/ConcertSync/ConcertSyncServer/Source/ConcertSyncServer/Private/Replication/SyncControlManager.h
+++ b/Engine/Plugins/Developer/Concert/ConcertSync/ConcertSyncServer/Source/ConcertSyncServer/Private/Replication/SyncControlManager.h
@@ -2,11 +2,11 @@
 
 #pragma once
 
-#include "MuteManager.h"
+#include "Muting/MuteManager.h"
 #include "Replication/Data/ObjectIds.h"
+#include "Replication/Messages/SyncControl.h"
 
 #include "HAL/Platform.h"
-#include "Replication/Messages/SyncControl.h"
 #include "Templates/UnrealTemplate.h"
 
 namespace UE::ConcertSyncServer::Replication
@@ -44,6 +44,8 @@ namespace UE::ConcertSyncServer::Replication
 
 		/** @return Whether this Object is allowed to processed. */
 		bool HasSyncControl(const FConcertReplicatedObjectId& Object) const;
+		/** @return If the client has any controlled objects, returns a valid pointer to the set of objects they own. */
+		const TSet<FConcertObjectInStreamID>* GetClientControlledObjects(const FGuid& ClientId) const;
 
 		/** Called by FConcertServerReplicationManager when client completes replication handshake. */
 		FConcertReplication_ChangeSyncControl OnGenerateSyncControlForClientJoin(const FGuid& ClientId);
diff --git a/Engine/Plugins/Developer/Concert/ConcertSync/ConcertSyncServer/Source/ConcertSyncServer/Private/Replication/Util/GroundTruthOverride.cpp b/Engine/Plugins/Developer/Concert/ConcertSync/ConcertSyncServer/Source/ConcertSyncServer/Private/Replication/Util/GroundTruthOverride.cpp
new file mode 100644
--- /dev/null
+++ b/Engine/Plugins/Developer/Concert/ConcertSync/ConcertSyncServer/Source/ConcertSyncServer/Private/Replication/Util/GroundTruthOverride.cpp
@@ -0,0 +1,46 @@
+﻿// Copyright Epic Games, Inc. All Rights Reserved.
+
+#include "GroundTruthOverride.h"
+
+#include "Replication/AuthorityManager.h"
+#include "Replication/ConcertReplicationClient.h"
+#include "Replication/Data/ReplicationStream.h"
+#include "Replication/Data/ReplicationStreamArray.h"
+
+namespace UE::ConcertSyncServer::Replication
+{
+	void FGroundTruthOverride::ForEachStream(const FGuid& ClientEndpointId, TFunctionRef<EBreakBehavior(const FGuid& StreamId, const FConcertObjectReplicationMap& ReplicationMap)> Callback) const
+	{
+		if (const FConcertReplicationStreamArray* StreamsOverride = StreamOverrides.Find(ClientEndpointId))
+		{
+			for (const FConcertReplicationStream& Stream : StreamsOverride->Streams)
+			{
+				if (Callback(Stream.BaseDescription.Identifier, Stream.BaseDescription.ReplicationMap) == EBreakBehavior::Break)
+				{
+					break;
+				}
+			}
+		}
+		else
+		{
+			Clients.ForEachStream(ClientEndpointId, [&Callback](const FConcertReplicationStream& Stream)
+			{
+				return Callback(Stream.BaseDescription.Identifier, Stream.BaseDescription.ReplicationMap);
+			});
+		}
+	}
+
+	void FGroundTruthOverride::ForEachSendingClient(TFunctionRef<EBreakBehavior(const FGuid& ClientEndpointId)> Callback) const
+	{
+		Clients.ForEachReplicationClient(Callback);
+	}
+
+	bool FGroundTruthOverride::HasAuthority(const FGuid& ClientId, const FGuid& StreamId, const FSoftObjectPath& ObjectPath) const
+	{
+		const FConcertReplicatedObjectId ObjectId {{ StreamId, ObjectPath }, ClientId };
+		const FConcertObjectInStreamArray* AuthorityOverride = AuthorityOverrides.Find(ClientId);
+		return AuthorityOverride
+			? AuthorityOverride->Objects.Contains(ObjectId)
+			: AuthorityManager.HasAuthorityToChange(ObjectId);
+	}
+}
diff --git a/Engine/Plugins/Developer/Concert/ConcertSync/ConcertSyncServer/Source/ConcertSyncServer/Private/Replication/Util/GroundTruthOverride.h b/Engine/Plugins/Developer/Concert/ConcertSync/ConcertSyncServer/Source/ConcertSyncServer/Private/Replication/Util/GroundTruthOverride.h
new file mode 100644
--- /dev/null
+++ b/Engine/Plugins/Developer/Concert/ConcertSync/ConcertSyncServer/Source/ConcertSyncServer/Private/Replication/Util/GroundTruthOverride.h
@@ -0,0 +1,61 @@
+﻿// Copyright Epic Games, Inc. All Rights Reserved.
+
+#pragma once
+
+#include "Replication/AuthorityConflictSharedUtils.h"
+
+#include "Containers/Array.h"
+#include "Containers/Map.h"
+#include "HAL/Platform.h"
+#include "Misc/Guid.h"
+#include "Templates/UnrealTemplate.h"
+
+struct FConcertReplicationStreamArray;
+struct FConcertObjectInStreamArray;
+
+namespace UE::ConcertSyncServer::Replication
+{
+	class FAuthorityManager;
+	class FConcertReplicationClient;
+	class IRegistrationEnumerator;
+	
+	/**
+	 * Pretends that the ground truth is the client overrides it was given.
+	 * If a setting is not overriden, then it defaults back to the server state.
+	 */
+	class FGroundTruthOverride
+		: public ConcertSyncCore::Replication::AuthorityConflictUtils::IReplicationGroundTruth
+		, public FNoncopyable
+	{
+	public:
+		
+		FGroundTruthOverride(
+			const TMap<FGuid, FConcertReplicationStreamArray>& StreamOverrides UE_LIFETIMEBOUND,
+			const TMap<FGuid, FConcertObjectInStreamArray>& AuthorityOverrides UE_LIFETIMEBOUND,
+			const IRegistrationEnumerator& Clients UE_LIFETIMEBOUND,
+			const FAuthorityManager& AuthorityManager UE_LIFETIMEBOUND
+			)
+			: StreamOverrides(StreamOverrides)
+			, AuthorityOverrides(AuthorityOverrides)
+			, Clients(Clients)
+			, AuthorityManager(AuthorityManager)
+		{
+		}
+
+		//~ Begin IReplicationGroundTruth Interface
+		virtual void ForEachStream(const FGuid& ClientEndpointId, TFunctionRef<EBreakBehavior(const FGuid& StreamId, const FConcertObjectReplicationMap& ReplicationMap)> Callback) const override;
+		virtual void ForEachSendingClient(TFunctionRef<EBreakBehavior(const FGuid& ClientEndpointId)> Callback) const override;
+		virtual bool HasAuthority(const FGuid& ClientId, const FGuid& StreamId, const FSoftObjectPath& ObjectPath) const override;
+		//~ Begin IReplicationGroundTruth Interface
+
+	private:
+
+		const TMap<FGuid, FConcertReplicationStreamArray>& StreamOverrides;
+		const TMap<FGuid, FConcertObjectInStreamArray>& AuthorityOverrides;
+
+		const IRegistrationEnumerator& Clients;
+		const FAuthorityManager& AuthorityManager;
+	};
+}
+
+
diff --git a/Engine/Plugins/Developer/Concert/ConcertSync/ConcertSyncServer/Source/ConcertSyncServer/Private/Replication/Util/StreamChangeValidation.cpp b/Engine/Plugins/Developer/Concert/ConcertSync/ConcertSyncServer/Source/ConcertSyncServer/Private/Replication/Util/StreamChangeValidation.cpp
new file mode 100644
--- /dev/null
+++ b/Engine/Plugins/Developer/Concert/ConcertSync/ConcertSyncServer/Source/ConcertSyncServer/Private/Replication/Util/StreamChangeValidation.cpp
@@ -0,0 +1,146 @@
+﻿// Copyright Epic Games, Inc. All Rights Reserved.
+
+#include "StreamChangeValidation.h"
+
+#include "ConcertLogGlobal.h"
+#include "Replication/AuthorityManager.h"
+#include "Replication/ChangeStreamSharedUtils.h"
+#include "Replication/ConcertReplicationClient.h"
+#include "Replication/Data/ReplicationStream.h"
+#include "Replication/Messages/ChangeStream.h"
+
+namespace UE::ConcertSyncServer::Replication
+{
+	namespace Private
+	{
+		static const FConcertReplicationStream* FindExistingStream(const FConcertReplicationClient& Client, const FGuid& StreamId)
+		{
+			return Client.GetStreamDescriptions().FindByPredicate([&StreamId](const FConcertReplicationStream& Description)
+			{
+				return Description.BaseDescription.Identifier == StreamId;
+			});
+		}
+		
+		/** Validates that ObjectsToPut writes only to pre-existing streams. */
+		static void ValidatePutObjectsRequestSemantics(const FConcertReplication_ChangeStream_Request& Request, const FConcertReplicationClient& Client, FConcertReplication_ChangeStream_Response& OutResponse)
+		{
+			const auto PutObjectHasEnoughData = [](const FConcertReplicationStream* ExistingStream, const FSoftObjectPath& ChangedObjectPath, const FConcertReplication_ChangeStream_PutObject& PutObject)
+			{
+				const bool bHasProperties = !PutObject.Properties.ReplicatedProperties.IsEmpty();
+				const bool bHasClassPath = !PutObject.ClassPath.IsNull();
+				const bool bIsEditingExistingObjectDefinition = ExistingStream->BaseDescription.ReplicationMap.ReplicatedObjects.Contains(ChangedObjectPath);
+				const bool bHasEnoughData = (!bIsEditingExistingObjectDefinition && bHasProperties && bHasClassPath)
+					|| (bIsEditingExistingObjectDefinition && bHasProperties)
+					|| (bIsEditingExistingObjectDefinition && bHasClassPath);
+				return bHasEnoughData;
+			};
+			
+			for (const TPair<FConcertObjectInStreamID, FConcertReplication_ChangeStream_PutObject>& Change : Request.ObjectsToPut)
+			{
+				const FConcertObjectInStreamID ChangedObject = Change.Key;
+				
+				const FGuid& StreamToModify = ChangedObject.StreamId;
+				const FConcertReplicationStream* ExistingStream = FindExistingStream(Client, StreamToModify);
+				const bool bStreamExists = ExistingStream != nullptr;
+				if (bStreamExists)
+				{
+					// FConcertReplicatedObjectInfo must have non-empty values for its members. Hence we must check that
+					// 1. if creating a new entry, both fields are given
+					// 2. if writing to pre-existing entry, at least one field is given
+					const FConcertReplication_ChangeStream_PutObject& PutObject = Change.Value; 
+					if (!PutObjectHasEnoughData(ExistingStream, ChangedObject.Object, PutObject))
+					{
+						OutResponse.ObjectsToPutSemanticErrors.Add(ChangedObject, EConcertPutObjectErrorCode::MissingData);
+					}
+				}
+				else
+				{
+					// ObjectsToPut can only write to pre-existing streams
+					UE_LOG(LogConcert, Log, TEXT("Semantic error: unknown stream %s"), *StreamToModify.ToString(EGuidFormats::Short));
+					OutResponse.ObjectsToPutSemanticErrors.Add(ChangedObject, EConcertPutObjectErrorCode::UnresolvedStream);
+				}
+			}
+		}
+
+		/** Checks that StreamsToAdd do not conflict with pre-existing ones and that all IDs in the request are also unique. */
+		static void ValidateAddedStreamsAreUnique(const FConcertReplication_ChangeStream_Request& Request, const FConcertReplicationClient& Client, FConcertReplication_ChangeStream_Response& OutResponse)
+		{
+			TSet<FGuid> DuplicateEntryDetection;
+			for (const FConcertReplicationStream& NewStream : Request.StreamsToAdd)
+			{
+				// StreamsToAdd is invalid if there is already a stream with the same ID registered ...
+				const FGuid& NewStreamId = NewStream.BaseDescription.Identifier;
+				const bool bIdAlreadyExists = FindExistingStream(Client, NewStreamId) != nullptr;
+				const bool bIsStreamRemoved = Request.StreamsToRemove.Contains(NewStreamId);
+
+				// ... or StreamsToAdd contains the same ID multiple times
+				bool bIsDuplicateEntry = false;
+				DuplicateEntryDetection.FindOrAdd(NewStreamId, &bIsDuplicateEntry);
+				
+				if ((bIdAlreadyExists && !bIsStreamRemoved)
+					|| bIsDuplicateEntry)
+				{
+					UE_LOG(LogConcert, Log, TEXT("Duplicate stream entry %s"), *NewStreamId.ToString(EGuidFormats::Short));
+					OutResponse.FailedStreamCreation.Add(NewStreamId);
+				}
+			}
+		}
+
+		/** If the requesting client has authority over a changed object, checks that no other client has authority over the properties being added. */
+		static void LookForAuthorityConflicts(const FConcertReplication_ChangeStream_Request& Request, const FConcertReplicationClient& Client, const FAuthorityManager& AuthorityManager, FConcertReplication_ChangeStream_Response& OutResponse)
+		{
+			const FGuid ClientEndpointId = Client.GetClientEndpointId();
+			for (const TPair<FConcertObjectInStreamID, FConcertReplication_ChangeStream_PutObject>& PutObjectPair : Request.ObjectsToPut)
+			{
+				// No conflict possible if requesting client does not have authority over the changed object
+				const FConcertReplicatedObjectId ReplicatedObjectInfo { { PutObjectPair.Key }, ClientEndpointId };
+				if (!AuthorityManager.HasAuthorityToChange(ReplicatedObjectInfo))
+				{
+					continue;
+				}
+
+				// Requester not changing properties (must be changing ClassPath)? Also no conflict possible.
+				const FConcertPropertySelection& PropertySelection = PutObjectPair.Value.Properties;
+				const bool bIsEditingProperties = !PropertySelection.ReplicatedProperties.IsEmpty();
+				if (!bIsEditingProperties)
+				{
+					continue;
+				}
+
+				// Simply check whether any other client is already sending any of the requested properties.
+				AuthorityManager.EnumerateAuthorityConflicts(ReplicatedObjectInfo, &PropertySelection,
+					[&OutResponse, &ReplicatedObjectInfo](const FGuid& ClientId, const FGuid& StreamId, const FConcertPropertyChain& Property)
+					{
+						const FConcertReplicatedObjectId ConflictingObject = { { StreamId, ReplicatedObjectInfo.Object }, ClientId };
+						OutResponse.AuthorityConflicts.Add(ReplicatedObjectInfo, ConflictingObject);
+						
+						UE_LOG(LogConcert, Log, TEXT("Authority conflict with client %s for stream %s for property %s"), *ClientId.ToString(EGuidFormats::Short), *StreamId.ToString(EGuidFormats::Short), *Property.ToString());
+						return EBreakBehavior::Continue;
+					});
+			}
+		}
+	}
+	
+	bool ValidateStreamChangeRequest(
+		const FConcertReplicationClient& Client,
+		const FAuthorityManager& AuthorityManager,
+		const FConcertReplication_ChangeStream_Request& Request,
+		FConcertReplication_ChangeStream_Response& OutResponse
+		)
+	{
+		OutResponse.ErrorCode = EReplicationResponseErrorCode::Handled;
+
+		Private::ValidatePutObjectsRequestSemantics(Request, Client, OutResponse);
+		Private::ValidateAddedStreamsAreUnique(Request, Client, OutResponse);
+		Private::LookForAuthorityConflicts(Request, Client, AuthorityManager, OutResponse);
+		ConcertSyncCore::Replication::ChangeStreamUtils::ValidateFrequencyChanges(Request, Client.GetStreamDescriptions(), &OutResponse.FrequencyErrors);
+			
+		return OutResponse.IsSuccess();
+	}
+
+	bool ValidateStreamChangeRequest(const FConcertReplicationClient& Client, const FAuthorityManager& AuthorityManager, const FConcertReplication_ChangeStream_Request& Request)
+	{
+		FConcertReplication_ChangeStream_Response Dummy;
+		return ValidateStreamChangeRequest(Client, AuthorityManager, Request, Dummy);
+	}
+}
diff --git a/Engine/Plugins/Developer/Concert/ConcertSync/ConcertSyncServer/Source/ConcertSyncServer/Private/Replication/Util/StreamChangeValidation.h b/Engine/Plugins/Developer/Concert/ConcertSync/ConcertSyncServer/Source/ConcertSyncServer/Private/Replication/Util/StreamChangeValidation.h
new file mode 100644
--- /dev/null
+++ b/Engine/Plugins/Developer/Concert/ConcertSync/ConcertSyncServer/Source/ConcertSyncServer/Private/Replication/Util/StreamChangeValidation.h
@@ -0,0 +1,32 @@
+﻿// Copyright Epic Games, Inc. All Rights Reserved.
+
+#pragma once
+
+struct FConcertReplication_ChangeStream_Response;
+struct FConcertReplication_ChangeStream_Request;
+
+namespace UE::ConcertSyncServer::Replication
+{
+	class FAuthorityManager;
+	class FConcertReplicationClient;
+
+	/**
+	 * Checks whether the stream change is valid to make and puts all errors into OutResponse.
+	 * @return Whether the request is valid.
+	 */
+	bool ValidateStreamChangeRequest(
+		const FConcertReplicationClient& Client,
+		const FAuthorityManager& AuthorityManager,
+		const FConcertReplication_ChangeStream_Request& Request,
+		FConcertReplication_ChangeStream_Response& OutResponse
+		);
+	
+	/** @return Whether the request is valid. */
+	bool ValidateStreamChangeRequest(
+		const FConcertReplicationClient& Client,
+		const FAuthorityManager& AuthorityManager,
+		const FConcertReplication_ChangeStream_Request& Request
+		);
+}
+
+
diff --git a/Engine/Plugins/Developer/Concert/ConcertSync/ConcertSyncTest/Source/ConcertSyncTest/Private/Replication/Components/SyncControlState.spec.cpp b/Engine/Plugins/Developer/Concert/ConcertSync/ConcertSyncTest/Source/ConcertSyncTest/Private/Replication/Components/SyncControlState.spec.cpp
--- a/Engine/Plugins/Developer/Concert/ConcertSync/ConcertSyncTest/Source/ConcertSyncTest/Private/Replication/Components/SyncControlState.spec.cpp
+++ b/Engine/Plugins/Developer/Concert/ConcertSync/ConcertSyncTest/Source/ConcertSyncTest/Private/Replication/Components/SyncControlState.spec.cpp
@@ -44,7 +44,7 @@ namespace UE::ConcertSyncTests::Replication::UI
 				TArray<FConcertObjectInStreamID> RemovedObjects;
 				SyncControl.AppendAuthorityChange(
 					Request,
-					Response,
+					Response.SyncControl,
 					[this](const FConcertObjectInStreamID&){ AddError(TEXT("No object should be added")); },
 					[&RemovedObjects](const FConcertObjectInStreamID& Object){ RemovedObjects.Add(Object); }
 					);
@@ -63,7 +63,7 @@ namespace UE::ConcertSyncTests::Replication::UI
 				TArray<FConcertObjectInStreamID> AddedObjects;
 				SyncControl.AppendAuthorityChange(
 					Request,
-					Response,
+					Response.SyncControl,
 					[&AddedObjects](const FConcertObjectInStreamID& Object){ AddedObjects.Add(Object); },
 					[this](const FConcertObjectInStreamID&){ AddError(TEXT("No object should be removed")); }
 					);
diff --git a/Engine/Plugins/Developer/Concert/ConcertSync/ConcertSyncTest/Source/ConcertSyncTest/Private/Replication/PutState/ChangeClientEvent.spec.cpp b/Engine/Plugins/Developer/Concert/ConcertSync/ConcertSyncTest/Source/ConcertSyncTest/Private/Replication/PutState/ChangeClientEvent.spec.cpp
new file mode 100644
--- /dev/null
+++ b/Engine/Plugins/Developer/Concert/ConcertSync/ConcertSyncTest/Source/ConcertSyncTest/Private/Replication/PutState/ChangeClientEvent.spec.cpp
@@ -0,0 +1,108 @@
+﻿// Copyright Epic Games, Inc. All Rights Reserved.
+
+#include "Replication/Messages/ChangeClientEvent.h"
+#include "Replication/Util/ConcertClientReplicationBridgeMock.h"
+#include "Replication/Util/Spec/ObjectTestReplicator.h"
+#include "Replication/Util/Spec/ReplicationClient.h"
+#include "Replication/Util/Spec/ReplicationServer.h"
+
+#include "Misc/AutomationTest.h"
+
+namespace UE::ConcertSyncTests::Replication::ChangeClients
+{
+	BEGIN_DEFINE_SPEC(FChangeClientEventSpec, "Editor.Concert.Replication.PutState.Event", EAutomationTestFlags::EditorContext | EAutomationTestFlags::EngineFilter)
+		TUniquePtr<FObjectTestReplicator> ObjectReplicator;
+
+		TUniquePtr<FReplicationServer> Server;
+		FReplicationClient* Client = nullptr;
+
+		const FGuid StreamId = FGuid::NewGuid();
+		FConcertReplicationStream StreamData;
+	END_DEFINE_SPEC(FChangeClientEventSpec);
+
+	/** This tests that FConcertReplication_ChangeClientEvent is correctly applied by the client when received. */
+	void FChangeClientEventSpec::Define()
+	{
+		BeforeEach([this]
+		{
+			ObjectReplicator = MakeUnique<FObjectTestReplicator>();
+			Server = MakeUnique<FReplicationServer>(*this);
+			Client = &Server->ConnectClient();
+
+			Client->JoinReplication();
+			StreamData = ObjectReplicator->CreateStream(StreamId);
+		});
+		AfterEach([this]
+		{
+			Server.Reset();
+			ObjectReplicator.Reset();
+		});
+
+		Describe("When client receives ChangeClientEvent", [this]
+		{
+			BeforeEach([this]
+			{
+				// This simulates the server sending to client - but it does not actually set any server state.
+				// This would not actually happen in a real world use case.
+				FConcertReplication_ChangeClientEvent Event;
+				Event.StreamChange.StreamsToAdd.Add(StreamData);
+				Event.AuthorityChange.TakeAuthority.Add({ ObjectReplicator->TestObject }, { { StreamId } });
+				Event.SyncControlChange.NewControlStates.Add({ StreamId, ObjectReplicator->TestObject }, true);
+				Server->GetServerSessionMock()->SendCustomEvent(Event, Client->GetEndpointId(), EConcertMessageFlags::ReliableOrdered);
+			});
+
+			It("Local client state is correct", [this]
+			{
+				const IConcertClientReplicationManager& ReplicationManager = Client->GetClientReplicationManager();
+				
+				const TArray<FConcertReplicationStream> Streams = ReplicationManager.GetRegisteredStreams();
+				if (Streams.Num() != 1)
+				{
+					AddError(TEXT("Expected 1 stream"));
+					return;
+				}
+				TestEqual(TEXT("Stream content"), Streams[0], StreamData);
+
+				const TMap<FSoftObjectPath, TSet<FGuid>> AuthorityObjects = ReplicationManager.GetClientOwnedObjects();
+				const TSet<FGuid>* AuthorityStreams = AuthorityObjects.Find(ObjectReplicator->TestObject);
+				if (!AuthorityStreams)
+				{
+					AddError(TEXT("Expected authority"));
+					return;
+				}
+				TestEqual(TEXT("AuthorityStreams.Num()"), AuthorityStreams->Num(), 1);
+				TestTrue(TEXT("AuthorityStreams.Contains(StreamId)"), AuthorityStreams->Contains(StreamId));
+				TestEqual(TEXT("AuthorityObjects.Num()"), AuthorityObjects.Num(), 1);
+
+				const TSet<FConcertObjectInStreamID> SyncControl = ReplicationManager.GetSyncControlledObjects();
+				TestEqual(TEXT("SyncControl.Num()"), SyncControl.Num(), 1);
+				TestTrue(TEXT("SyncControl.Contains(TestObject)"), SyncControl.Contains({ StreamId, ObjectReplicator->TestObject }));
+			});
+
+			It("Attempts to replicate", [this]
+			{
+				FReplicationClient& Receiver = Server->ConnectClient();
+				// Receiver won't actually receive anything because server will block it.
+				Receiver.JoinReplicationAsListener({});
+
+				bool bSentDataToServer = false;
+				Client->GetBridgeMock().InjectAvailableObject(*ObjectReplicator->TestObject);
+				ObjectReplicator->SimulateSendObjectToReceiver(
+					*this,
+					FObjectReplicationContext{ *Client, *Server, Receiver },
+					{ StreamId },
+					[&bSentDataToServer](const FConcertSessionContext& Context, const FConcertReplication_BatchReplicationEvent& Event)
+					{
+						bSentDataToServer = true;
+					},
+					[this](const FConcertSessionContext& Context, const FConcertReplication_BatchReplicationEvent& Event)
+					{
+						// The server is supposed to reject the data because our tests did not send the FConcertReplication_ChangeClientEvent through the replication system.
+						AddError(TEXT("Server forwarded the replicated data!"));
+					});
+				TestTrue(TEXT("Client attempted to replicate"), bSentDataToServer);
+				ObjectReplicator->TestValuesWereNotReplicated(*this);
+			});
+		});
+	}
+}
diff --git a/Engine/Plugins/Developer/Concert/ConcertSync/ConcertSyncTest/Source/ConcertSyncTest/Private/Replication/PutState/PutStateAndMute.spec.cpp b/Engine/Plugins/Developer/Concert/ConcertSync/ConcertSyncTest/Source/ConcertSyncTest/Private/Replication/PutState/PutStateAndMute.spec.cpp
new file mode 100644
--- /dev/null
+++ b/Engine/Plugins/Developer/Concert/ConcertSync/ConcertSyncTest/Source/ConcertSyncTest/Private/Replication/PutState/PutStateAndMute.spec.cpp
@@ -0,0 +1,311 @@
+﻿// Copyright Epic Games, Inc. All Rights Reserved.
+
+#include "Replication/Util/Spec/ObjectTestReplicator.h"
+#include "Replication/Util/Spec/ReplicationClient.h"
+#include "Replication/Util/Spec/ReplicationServer.h"
+
+#include "Misc/AutomationTest.h"
+#include "Replication/Messages/ChangeClientEvent.h"
+#include "Replication/Util/ConcertClientReplicationBridgeMock.h"
+
+namespace UE::ConcertSyncTests::Replication::ChangeClients
+{
+	BEGIN_DEFINE_SPEC(FPutStateAndMuteSpec, "Editor.Concert.Replication.PutState.Mute", EAutomationTestFlags::EditorContext | EAutomationTestFlags::EngineFilter)
+		TUniquePtr<FObjectTestReplicator> ObjectReplicator;
+
+		TUniquePtr<FReplicationServer> Server;
+		FReplicationClient* Sender = nullptr;
+		FReplicationClient* Receiver = nullptr;
+
+		const FGuid StreamId = FGuid::NewGuid();
+	END_DEFINE_SPEC(FPutStateAndMuteSpec);
+
+	/** This tests that muting works after a successful FConcertReplication_PutState change. */
+	void FPutStateAndMuteSpec::Define()
+	{
+		BeforeEach([this]
+		{
+			ObjectReplicator = MakeUnique<FObjectTestReplicator>();
+			Server = MakeUnique<FReplicationServer>(*this);
+			Sender = &Server->ConnectClient();
+			Receiver = &Server->ConnectClient();
+
+			Sender->JoinReplication();
+			Receiver->JoinReplicationAsListener({ ObjectReplicator->TestObject });
+		});
+		AfterEach([this]
+		{
+			Server.Reset();
+			ObjectReplicator.Reset();
+		});
+
+		Describe("When client adds stream and mutes it", [this]
+		{
+			const auto RunTest = [this]()
+			{
+				IConcertClientReplicationManager& ReplicationManager = Sender->GetClientReplicationManager();
+				FConcertReplication_PutState_Request Request;
+				Request.NewStreams.Add(Sender->GetEndpointId(), { { ObjectReplicator->CreateStream(StreamId) } });
+				Request.NewAuthorityState.Add(Sender->GetEndpointId()).Objects.Add({ StreamId, ObjectReplicator->TestObject });
+				Request.MuteChange.ObjectsToMute.Add(ObjectReplicator->TestObject);
+
+				bool bReceivedResponse = false;
+				TFuture<FConcertReplication_PutState_Response> Future = ReplicationManager
+					.PutClientState(Request)
+					.Next([this, &bReceivedResponse](FConcertReplication_PutState_Response&& Response)
+					{
+						bReceivedResponse = true;
+						TestTrue(TEXT("Success"), Response.IsSuccess());
+						TestTrue(TEXT("No sync control"), Response.SyncControl.IsEmpty());
+						return Response;
+					});
+				TestTrue(TEXT("bReceivedResponse"), bReceivedResponse);
+				return Future;
+			};
+			
+			It("The response does not contain change for sync control", [this, RunTest]
+			{
+				RunTest().Next([this](FConcertReplication_PutState_Response&& Response)
+				{
+					TestTrue(TEXT("No sync control in response"), Response.SyncControl.IsEmpty());
+				});
+			});
+			It("The client predicts losing sync control", [this, RunTest]
+			{
+				RunTest();
+				
+				IConcertClientReplicationManager& ReplicationManager = Sender->GetClientReplicationManager();
+				TestEqual(TEXT("GetSyncControlledObjects().Num()"), ReplicationManager.GetSyncControlledObjects().Num(), 0);
+			});
+			It("The object is globally muted", [this, RunTest]
+			{
+				RunTest();
+				Sender->GetClientReplicationManager()
+					.QueryMuteState()
+					.Next([this](FConcertReplication_QueryMuteState_Response&& Response)
+					{
+						TestTrue(TEXT("Object is muted"), Response.ExplicitlyMutedObjects.Contains(ObjectReplicator->TestObject));
+						TestEqual(TEXT("ExplicitlyMutedObjects.Num()"), Response.ExplicitlyMutedObjects.Num(), 1);
+						TestEqual(TEXT("ExplicitlyUnmutedObjects.Num()"), Response.ExplicitlyUnmutedObjects.Num(), 0);
+						TestEqual(TEXT("ImplicitlyMutedObjects.Num()"), Response.ImplicitlyMutedObjects.Num(), 0);
+						TestEqual(TEXT("ImplicitlyUnmutedObjects.Num()"), Response.ImplicitlyUnmutedObjects.Num(), 0);
+					});
+			});
+			It("The object cannot be replicated", [this, RunTest]
+			{
+				RunTest();
+
+				Sender->GetBridgeMock().InjectAvailableObject(*ObjectReplicator->TestObject);
+				ObjectReplicator->SimulateSendObjectToReceiver(*this, { *Sender, *Server, *Receiver }, { StreamId });
+				ObjectReplicator->TestValuesWereNotReplicated(*this);
+			});
+		});
+
+		Describe("When a client had sync control and mutes object using put state", [this]
+		{
+			const auto RunTest = [this]()
+			{
+				IConcertClientReplicationManager& ReplicationManager = Sender->GetClientReplicationManager();
+				ReplicationManager.ChangeStream({ .StreamsToAdd = { ObjectReplicator->CreateStream(StreamId) } });
+				ReplicationManager.TakeAuthorityOver({ { ObjectReplicator->TestObject } });
+				
+				bool bReceivedResponse = false;
+				TFuture<FConcertReplication_PutState_Response> Future = ReplicationManager
+					.PutClientState({ .MuteChange = { .ObjectsToMute = { { ObjectReplicator->TestObject, {} } } } })
+					.Next([this, &bReceivedResponse](FConcertReplication_PutState_Response&& Response)
+					{
+						bReceivedResponse = true;
+						TestTrue(TEXT("Success"), Response.IsSuccess());
+						return Response;
+					});
+				TestTrue(TEXT("bReceivedResponse"), bReceivedResponse);
+				return Future;
+			};
+
+			It("The client predicts losing sync control", [this, RunTest]
+			{
+				RunTest();
+				
+				IConcertClientReplicationManager& ReplicationManager = Sender->GetClientReplicationManager();
+				TestEqual(TEXT("GetSyncControlledObjects().Num()"), ReplicationManager.GetSyncControlledObjects().Num(), 0);
+			});
+			It("The response does not contain change for sync control", [this, RunTest]
+			{
+				RunTest().Next([this](FConcertReplication_PutState_Response&& Response)
+				{
+					TestTrue(TEXT("No sync control in response"), Response.SyncControl.IsEmpty());
+				});
+			});
+			It("The object is globally muted", [this, RunTest]
+			{
+				RunTest();
+				Sender->GetClientReplicationManager()
+					.QueryMuteState()
+					.Next([this](FConcertReplication_QueryMuteState_Response&& Response)
+					{
+						TestTrue(TEXT("Object is muted"), Response.ExplicitlyMutedObjects.Contains(ObjectReplicator->TestObject));
+						TestEqual(TEXT("ExplicitlyMutedObjects.Num()"), Response.ExplicitlyMutedObjects.Num(), 1);
+						TestEqual(TEXT("ExplicitlyUnmutedObjects.Num()"), Response.ExplicitlyUnmutedObjects.Num(), 0);
+						TestEqual(TEXT("ImplicitlyMutedObjects.Num()"), Response.ImplicitlyMutedObjects.Num(), 0);
+						TestEqual(TEXT("ImplicitlyUnmutedObjects.Num()"), Response.ImplicitlyUnmutedObjects.Num(), 0);
+					});
+			});
+			It("The object cannot be replicated", [this, RunTest]
+			{
+				RunTest();
+				
+				Sender->GetBridgeMock().InjectAvailableObject(*ObjectReplicator->TestObject);
+				ObjectReplicator->SimulateSendObjectToReceiver(*this, { *Sender, *Server, *Receiver }, { StreamId });
+				ObjectReplicator->TestValuesWereNotReplicated(*this);
+			});
+		});
+		
+		Describe("When client adds a stream and unmutes it", [this]
+		{
+			const auto RunTest = [this]()
+			{
+				// In order to mute an object, it must be known to the server. So we must create a dummy stream before the put request.
+				IConcertClientReplicationManager& ReceiverManager = Receiver->GetClientReplicationManager();
+				ReceiverManager.ChangeStream({ .StreamsToAdd = { { ObjectReplicator->CreateStream(StreamId) }} });
+				ReceiverManager.MuteObjects({ ObjectReplicator->TestObject });
+				
+				IConcertClientReplicationManager& SenderManager = Sender->GetClientReplicationManager();
+				FConcertReplication_PutState_Request Request;
+				Request.NewStreams.Add(Sender->GetEndpointId(), { .Streams = { ObjectReplicator->CreateStream(StreamId) } });
+				Request.NewStreams.Add(Receiver->GetEndpointId(), {});
+				Request.NewAuthorityState.Add(Sender->GetEndpointId(), { .Objects = {{ StreamId, ObjectReplicator->TestObject }} });
+				Request.MuteChange.ObjectsToUnmute.Add(ObjectReplicator->TestObject);
+				bool bReceivedResponse = false;
+				TFuture<FConcertReplication_PutState_Response> Future = SenderManager
+					.PutClientState(Request)
+					.Next([this, &bReceivedResponse](FConcertReplication_PutState_Response&& Response)
+					{
+						bReceivedResponse = true;
+						TestTrue(TEXT("Success"), Response.IsSuccess());
+						return Response;
+					});
+				TestTrue(TEXT("bReceivedResponse"), bReceivedResponse);
+				return Future;
+			};
+			
+			It("The response contains sync control", [this, RunTest]
+			{
+				RunTest().Next([this](FConcertReplication_PutState_Response&& Response)
+				{
+					TestEqual(TEXT("NewControlStates.Num()"), Response.SyncControl.NewControlStates.Num(), 1);
+					const bool* SyncControlState = Response.SyncControl.NewControlStates.Find({ StreamId, ObjectReplicator->TestObject });
+					TestTrue(TEXT("Has Sync Control"), SyncControlState && *SyncControlState);
+				});
+			});
+			It("The client thinks it has sync control", [this, RunTest]
+			{
+				RunTest();
+				
+				IConcertClientReplicationManager& SenderManager = Sender->GetClientReplicationManager();
+				TestEqual(TEXT("GetSyncControlledObjects().Num()"), SenderManager.GetSyncControlledObjects().Num(), 1);
+				TestTrue(TEXT("GetSyncControlledObjects().Contains(TestObject)"), SenderManager.GetSyncControlledObjects().Contains({ StreamId, ObjectReplicator->TestObject }));
+			});
+			It("The object is not globally muted", [this, RunTest]
+			{
+				RunTest();
+				Sender->GetClientReplicationManager()
+					.QueryMuteState()
+					.Next([this](FConcertReplication_QueryMuteState_Response&& Response)
+					{
+						TestEqual(TEXT("ExplicitlyMutedObjects.Num()"), Response.ExplicitlyMutedObjects.Num(), 0);
+						TestEqual(TEXT("ExplicitlyUnmutedObjects.Num()"), Response.ExplicitlyUnmutedObjects.Num(), 0);
+						TestEqual(TEXT("ImplicitlyMutedObjects.Num()"), Response.ImplicitlyMutedObjects.Num(), 0);
+						TestEqual(TEXT("ImplicitlyUnmutedObjects.Num()"), Response.ImplicitlyUnmutedObjects.Num(), 0);
+					});
+			});
+			It("The object can be replicated", [this, RunTest]
+			{
+				RunTest();
+				
+				Sender->GetBridgeMock().InjectAvailableObject(*ObjectReplicator->TestObject);
+				ObjectReplicator->SimulateSendObjectToReceiver(*this, { *Sender, *Server, *Receiver }, { StreamId });
+				ObjectReplicator->TestValuesWereReplicated(*this);
+			});
+		});
+
+		Describe("Client receives FConcertReplication_ChangeClientEvent", [this]
+		{
+			BeforeEach([this]
+			{
+				IConcertClientReplicationManager& ReplicationManager = Sender->GetClientReplicationManager();
+				ReplicationManager.ChangeStream({ .StreamsToAdd = {{ ObjectReplicator->CreateStream(StreamId) }} });
+				ReplicationManager.TakeAuthorityOver({ ObjectReplicator->TestObject });
+				ensureMsgf(ReplicationManager.GetSyncControlledObjects().Contains({ StreamId, ObjectReplicator->TestObject }), TEXT("Test not set up correctly."));
+			});
+			
+			It("When client is loses sync control due to a mute change.", [this]
+			{
+				int32 EventCount = 0;
+				Sender->GetClientSessionMock()->RegisterCustomEventHandler<FConcertReplication_ChangeClientEvent>(
+					[this, &EventCount](const FConcertSessionContext&, const FConcertReplication_ChangeClientEvent& Event)
+					{
+						++EventCount;
+						
+						TestEqual(TEXT("NewControlStates.Num()"), Event.SyncControlChange.NewControlStates.Num(), 1);
+						const bool* NewSyncControl = Event.SyncControlChange.NewControlStates.Find({ StreamId, ObjectReplicator->TestObject });
+						TestTrue(TEXT("Has No Sync Control"), NewSyncControl && !*NewSyncControl);
+					});
+				
+				IConcertClientReplicationManager& ReplicationManager = Receiver->GetClientReplicationManager();
+				FConcertReplication_PutState_Request Request;
+				Request.MuteChange.ObjectsToMute.Add(ObjectReplicator->TestObject);
+				ReplicationManager.PutClientState(Request);
+
+				TestEqual(TEXT("EventCount"), EventCount, 1);
+			});
+			It("When client is gains sync control due to a mute change.", [this]
+			{
+				int32 EventCount = 0;
+				Sender->GetClientSessionMock()->RegisterCustomEventHandler<FConcertReplication_ChangeClientEvent>(
+					[this, &EventCount](const FConcertSessionContext&, const FConcertReplication_ChangeClientEvent& Event)
+					{
+						++EventCount;
+						
+						TestEqual(TEXT("NewControlStates.Num()"), Event.SyncControlChange.NewControlStates.Num(), 1);
+						const bool* NewSyncControl = Event.SyncControlChange.NewControlStates.Find({ StreamId, ObjectReplicator->TestObject });
+						TestTrue(TEXT("Has No Sync Control"), NewSyncControl && *NewSyncControl);
+					});
+				
+				IConcertClientReplicationManager& ReplicationManager = Receiver->GetClientReplicationManager();
+				ReplicationManager.MuteObjects({ ObjectReplicator->TestObject });
+				FConcertReplication_PutState_Request Request;
+				Request.MuteChange.ObjectsToUnmute.Add(ObjectReplicator->TestObject);
+				ReplicationManager.PutClientState(Request);
+
+				TestEqual(TEXT("EventCount"), EventCount, 1);
+			});
+		});
+		
+		It("When client is not affected by a mute change, it does not receive FConcertReplication_ChangeClientEvent.", [this]
+		{
+			int32 EventCount = 0;
+			Receiver->GetClientSessionMock()->RegisterCustomEventHandler<FConcertReplication_ChangeClientEvent>(
+				[this, &EventCount](const FConcertSessionContext&, const FConcertReplication_ChangeClientEvent& Event)
+				{
+					++EventCount;
+					
+					TestEqual(TEXT("NewControlStates.Num()"), Event.SyncControlChange.NewControlStates.Num(), 1);
+					const bool* NewSyncControl = Event.SyncControlChange.NewControlStates.Find({ StreamId, ObjectReplicator->TestObject });
+					TestTrue(TEXT("Has No Sync Control"), NewSyncControl && *NewSyncControl);
+				});
+			
+			FObjectTestReplicator BarReplicator;
+			IConcertClientReplicationManager& ReceiverManager = Receiver->GetClientReplicationManager();
+			ReceiverManager.ChangeStream({ .StreamsToAdd = { BarReplicator.CreateStream(StreamId) } });
+			ReceiverManager.TakeAuthorityOver({ BarReplicator.TestObject });
+			
+			IConcertClientReplicationManager& SenderManager = Sender->GetClientReplicationManager();
+			FConcertReplication_PutState_Request Request;
+			Request.NewStreams.Add(Sender->GetEndpointId(), {{ ObjectReplicator->CreateStream(StreamId) }});
+			Request.MuteChange.ObjectsToMute.Add(ObjectReplicator->TestObject);
+			SenderManager.PutClientState(Request);
+
+			TestEqual(TEXT("EventCount"), EventCount, 0);
+		});
+	}
+}
diff --git a/Engine/Plugins/Developer/Concert/ConcertSync/ConcertSyncTest/Source/ConcertSyncTest/Private/Replication/PutState/PutStateAndReplicate.spec.cpp b/Engine/Plugins/Developer/Concert/ConcertSync/ConcertSyncTest/Source/ConcertSyncTest/Private/Replication/PutState/PutStateAndReplicate.spec.cpp
new file mode 100644
--- /dev/null
+++ b/Engine/Plugins/Developer/Concert/ConcertSync/ConcertSyncTest/Source/ConcertSyncTest/Private/Replication/PutState/PutStateAndReplicate.spec.cpp
@@ -0,0 +1,158 @@
+﻿// Copyright Epic Games, Inc. All Rights Reserved.
+
+#include "Replication/Util/ConcertClientReplicationBridgeMock.h"
+#include "Replication/Util/Spec/ObjectTestReplicator.h"
+#include "Replication/Util/Spec/ReplicationClient.h"
+#include "Replication/Util/Spec/ReplicationServer.h"
+
+#include "Misc/AutomationTest.h"
+
+namespace UE::ConcertSyncTests::Replication::ChangeClients
+{
+	BEGIN_DEFINE_SPEC(FPutStateAndReplicateSpec, "Editor.Concert.Replication.PutState.Replicate", EAutomationTestFlags::EditorContext | EAutomationTestFlags::EngineFilter)
+		TUniquePtr<FObjectTestReplicator> ObjectReplicator;
+
+		TUniquePtr<FReplicationServer> Server;
+		FReplicationClient* OldSender = nullptr;
+		FReplicationClient* NewSender = nullptr;
+		FReplicationClient* Receiver = nullptr;
+
+		const FGuid StreamId = FGuid::NewGuid();
+		FConcertReplicationStream StreamData;
+	
+		FObjectReplicationContext MakeReplicationContext(FReplicationClient& Sender) const { return { Sender, *Server, *Receiver }; }
+
+		FConcertReplication_PutState_Request BuildRequest()
+		{
+			FConcertReplication_PutState_Request Request;
+
+			FConcertReplicationStreamArray OldSenderStreamContent { };
+			Request.NewStreams.Add(OldSender->GetEndpointId(), OldSenderStreamContent);
+
+			FConcertReplicationStreamArray NewSenderStreamContent { { StreamData } };
+			Request.NewStreams.Add(NewSender->GetEndpointId(), NewSenderStreamContent);
+
+			const FConcertObjectInStreamID ObjectId = { StreamId, ObjectReplicator->TestObject };
+			Request.NewAuthorityState.Add(NewSender->GetEndpointId(), { { ObjectId } });
+			
+			return Request;
+		}
+	
+		void BuildReplicationCases()
+		{
+			It("OldSender has empty state", [this]
+			{
+				const IConcertClientReplicationManager& ReplicationManager = OldSender->GetClientReplicationManager();
+				TestEqual(TEXT("RegisteredStreams.Num()"), ReplicationManager.GetRegisteredStreams().Num(), 0);
+				TestEqual(TEXT("ClientOwnedObjects.Num()"), ReplicationManager.GetClientOwnedObjects().Num(), 0);
+				TestEqual(TEXT("SyncControl.Num()"), ReplicationManager.GetSyncControlledObjects().Num(), 0);
+			});
+			
+			It("NewSender has correct state", [this]
+			{
+				const IConcertClientReplicationManager& ReplicationManager = NewSender->GetClientReplicationManager();
+				
+				const TArray<FConcertReplicationStream> Streams = ReplicationManager.GetRegisteredStreams();
+				if (Streams.Num() != 1)
+				{
+					AddError(TEXT("Expected 1 stream"));
+					return;
+				}
+				TestEqual(TEXT("Stream content"), Streams[0], StreamData);
+
+				const TMap<FSoftObjectPath, TSet<FGuid>> AuthorityObjects = ReplicationManager.GetClientOwnedObjects();
+				const TSet<FGuid>* AuthorityStreams = AuthorityObjects.Find(ObjectReplicator->TestObject);
+				if (!AuthorityStreams)
+				{
+					AddError(TEXT("Expected authority"));
+					return;
+				}
+				TestEqual(TEXT("AuthorityStreams.Num()"), AuthorityStreams->Num(), 1);
+				TestTrue(TEXT("AuthorityStreams.Contains(StreamId)"), AuthorityStreams->Contains(StreamId));
+				TestEqual(TEXT("AuthorityObjects.Num()"), AuthorityObjects.Num(), 1);
+
+				const TSet<FConcertObjectInStreamID> SyncControl = ReplicationManager.GetSyncControlledObjects();
+				TestEqual(TEXT("SyncControl.Num()"), SyncControl.Num(), 1);
+				TestTrue(TEXT("SyncControl.Contains(TestObject)"), SyncControl.Contains({ StreamId, ObjectReplicator->TestObject }));
+			});
+
+			It("OldSender cannot replicate", [this]
+			{
+				ObjectReplicator->SimulateSendObjectToReceiver(*this, { *OldSender, *Server, *Receiver }, { StreamId });
+				ObjectReplicator->TestValuesWereNotReplicated(*this);
+			});
+			
+			It("NewSender can replicate", [this]
+			{
+				ObjectReplicator->SimulateSendObjectToReceiver(*this, { *NewSender, *Server, *Receiver }, { StreamId });
+				ObjectReplicator->TestValuesWereReplicated(*this);
+			});
+		}
+	END_DEFINE_SPEC(FPutStateAndReplicateSpec);
+
+	/**
+	 * This tests that replication works after a successful FConcertReplication_PutState change.
+	 * It transfers authority from OldSender to NewSender and ensures that replication to Receiver still works.
+	 */
+	void FPutStateAndReplicateSpec::Define()
+	{
+		BeforeEach([this]
+		{
+			ObjectReplicator = MakeUnique<FObjectTestReplicator>();
+			Server = MakeUnique<FReplicationServer>(*this);
+			OldSender = &Server->ConnectClient();
+			NewSender = &Server->ConnectClient();
+			Receiver = &Server->ConnectClient();
+
+			StreamData = ObjectReplicator->CreateStream(StreamId);
+			OldSender->JoinReplication(ObjectReplicator->CreateSenderArgs(StreamId));
+			NewSender->JoinReplication();
+			Receiver->JoinReplicationAsListener({ ObjectReplicator->TestObject });
+			
+			OldSender->GetClientReplicationManager().TakeAuthorityOver({ ObjectReplicator->TestObject });
+			OldSender->GetBridgeMock().InjectAvailableObject(*ObjectReplicator->TestObject);
+			NewSender->GetBridgeMock().InjectAvailableObject(*ObjectReplicator->TestObject);
+
+			// Just double-checking that we set up the test correctly...
+			ObjectReplicator->SimulateSendObjectToReceiver(*this, MakeReplicationContext(*OldSender), { StreamId });
+			ObjectReplicator->TestValuesWereReplicated(*this);
+		});
+		AfterEach([this]
+		{
+			Server.Reset();
+			ObjectReplicator.Reset();
+		});
+
+		Describe("When Receiver transfers from OldSender to NewSender", [this]
+		{
+			BeforeEach([this]
+			{
+				Receiver->GetClientReplicationManager().PutClientState(BuildRequest());
+			});
+			
+			BuildReplicationCases();
+		});
+		
+		Describe("When NewSender transfers from OldSender to NewSender", [this]
+		{
+			BeforeEach([this]
+			{
+				bool bReceivedResponse = false;
+				NewSender->GetClientReplicationManager()
+					.PutClientState(BuildRequest())
+					.Next([this, &bReceivedResponse](FConcertReplication_PutState_Response&& Response)
+					{
+						bReceivedResponse = true;
+						TestTrue(TEXT("IsSuccess"), Response.IsSuccess());
+						
+						const bool* pbHasSyncControl = Response.SyncControl.NewControlStates.Find({ StreamId, ObjectReplicator->TestObject });
+						TestTrue(TEXT("Has Sync Control"), pbHasSyncControl && *pbHasSyncControl);
+						TestEqual(TEXT("SyncControl.Num()"), Response.SyncControl.NewControlStates.Num(), 1);
+					});
+				TestTrue(TEXT("bReceivedResponse"), bReceivedResponse);
+			});
+			
+			BuildReplicationCases();
+		});
+	}
+}
diff --git a/Engine/Plugins/Developer/Concert/ConcertSync/ConcertSyncTest/Source/ConcertSyncTest/Private/Replication/PutState/PutStateDisabled.spec.cpp b/Engine/Plugins/Developer/Concert/ConcertSync/ConcertSyncTest/Source/ConcertSyncTest/Private/Replication/PutState/PutStateDisabled.spec.cpp
new file mode 100644
--- /dev/null
+++ b/Engine/Plugins/Developer/Concert/ConcertSync/ConcertSyncTest/Source/ConcertSyncTest/Private/Replication/PutState/PutStateDisabled.spec.cpp
@@ -0,0 +1,72 @@
+﻿// Copyright Epic Games, Inc. All Rights Reserved.
+
+#include "Replication/Util/Spec/ReplicationClient.h"
+#include "Replication/Util/Spec/ReplicationServer.h"
+#include "Replication/Util/Spec/ObjectTestReplicator.h"
+
+#include "Misc/AutomationTest.h"
+#include "Replication/Messages/PutState.h"
+
+namespace UE::ConcertSyncTests::Replication::ChangeClients
+{
+	BEGIN_DEFINE_SPEC(FPutStateDisabledSpec, "Editor.Concert.Replication.PutState", EAutomationTestFlags::EditorContext | EAutomationTestFlags::EngineFilter)
+		const FGuid ClientId { 0, 0, 0, 1 };
+	
+		void TestRequestHasErrorCode(EConcertReplicationPutStateResponseCode ExpectedErrorCode, EConcertSyncSessionFlags Flags, FConcertReplication_PutState_Request Request = {})
+		{
+			TUniquePtr<FReplicationServer> Server = MakeUnique<FReplicationServer>(*this, Flags);
+			FReplicationClient& Client = Server->ConnectClient();
+			check(Client.GetEndpointId() == ClientId);
+			Client.JoinReplication();
+			
+			bool bReceivedResponse = false;
+			Client.GetClientReplicationManager()
+				.PutClientState(Request)
+				.Next([this, &bReceivedResponse, ExpectedErrorCode](FConcertReplication_PutState_Response&& Response)
+				{
+					bReceivedResponse = true;
+					TestEqual(TEXT("Error code"), Response.ResponseCode, ExpectedErrorCode);
+				});
+			TestTrue(TEXT("Received response"), bReceivedResponse);
+		}
+	END_DEFINE_SPEC(FPutStateDisabledSpec);
+
+	/**
+	 * This tests that FConcertReplication_PutState_Request works as intended depending on the session's EConcertSyncSessionFlags flags.
+	 * The request should be rejected when EConcertSyncSessionFlags::ShouldEnableRemoteEditing is not set.
+	 */
+	void FPutStateDisabledSpec::Define()
+	{
+		It("When EConcertSyncSessionFlags::ShouldEnableRemoteEditing is not set, then FConcertReplication_PutState_Request fails.", [this]
+		{
+			TestRequestHasErrorCode(
+				EConcertReplicationPutStateResponseCode::FeatureDisabled,
+				EConcertSyncSessionFlags::Default_MultiUserSession & ~EConcertSyncSessionFlags::ShouldEnableRemoteEditing
+			);
+		});
+		
+		Describe("When EConcertSyncSessionFlags::ShouldEnableRemoteEditing is set but EConcertSyncSessionFlags::ShouldAllowGlobalMuting is not", [this]
+		{
+			It("When mute request is empty, the request succeeds", [this]
+			{
+				TestRequestHasErrorCode(
+					EConcertReplicationPutStateResponseCode::Success,
+					EConcertSyncSessionFlags::Default_MultiUserSession & ~EConcertSyncSessionFlags::ShouldAllowGlobalMuting
+				);
+			});
+			It("When mute request is non-empty, the request fails", [this]
+			{
+				FObjectTestReplicator ObjectReplicator;
+				FConcertReplication_PutState_Request Request;
+				Request.NewStreams.Add(ClientId, {{ ObjectReplicator.CreateStream() }});
+				Request.MuteChange.ObjectsToMute.Add({ ObjectReplicator.TestObject });
+				
+				TestRequestHasErrorCode(
+					EConcertReplicationPutStateResponseCode::FeatureDisabled,
+					EConcertSyncSessionFlags::Default_MultiUserSession & ~EConcertSyncSessionFlags::ShouldAllowGlobalMuting,
+					Request
+				);
+			});
+		});
+	}
+}
diff --git a/Engine/Plugins/Developer/Concert/ConcertSync/ConcertSyncTest/Source/ConcertSyncTest/Private/Replication/PutState/PutStateRejection.spec.cpp b/Engine/Plugins/Developer/Concert/ConcertSync/ConcertSyncTest/Source/ConcertSyncTest/Private/Replication/PutState/PutStateRejection.spec.cpp
new file mode 100644
--- /dev/null
+++ b/Engine/Plugins/Developer/Concert/ConcertSync/ConcertSyncTest/Source/ConcertSyncTest/Private/Replication/PutState/PutStateRejection.spec.cpp
@@ -0,0 +1,269 @@
+﻿// Copyright Epic Games, Inc. All Rights Reserved.
+
+#include "Replication/Util/ConcertClientReplicationBridgeMock.h"
+#include "Replication/Util/Spec/ObjectTestReplicator.h"
+#include "Replication/Util/Spec/ReplicationClient.h"
+#include "Replication/Util/Spec/ReplicationServer.h"
+
+#include "Misc/AutomationTest.h"
+
+namespace UE::ConcertSyncTests::Replication::ChangeClients
+{
+	BEGIN_DEFINE_SPEC(FPutStateRejectionSpec, "Editor.Concert.Replication.PutState.Rejection", EAutomationTestFlags::EditorContext | EAutomationTestFlags::EngineFilter)
+		TUniquePtr<FObjectTestReplicator> ObjectReplicator;
+
+		TUniquePtr<FReplicationServer> Server;
+		FReplicationClient* Client1 = nullptr;
+		FReplicationClient* Client2 = nullptr;
+	
+		const FGuid StreamId = FGuid::NewGuid();
+		FConcertReplicationStream StreamData;
+		FConcertReplicationStream StreamData_FloatOnly;
+		FConcertReplicationStream StreamData_VectorOnly;
+
+		const FConcertReplicatedObjectInfo& GetReplicatedObjectInfo() const { return StreamData.BaseDescription.ReplicationMap.ReplicatedObjects[ObjectReplicator->TestObject]; }
+		const TSet<FConcertPropertyChain>& GetReplicatedProperties() const { return GetReplicatedObjectInfo().PropertySelection.ReplicatedProperties; }
+
+		void TestClientHasConflictWith(const FConcertReplication_PutState_Response& Response, const FGuid& ClientThatHasConflict, const FGuid& ClientConflictingWith)
+		{
+			const FConcertAuthorityConflictArray* AuthorityErrors = Response.AuthorityChangeConflicts.Find(ClientThatHasConflict);
+			if (!AuthorityErrors)
+			{
+				AddError(TEXT("No error generated"));
+				return;
+			}
+
+			if (AuthorityErrors->Conflicts.Num() != 1)
+			{
+				AddError(TEXT("Expected exactly 1 error"));
+				return;
+			}
+
+			const FConcertAuthorityConflict& Conflict = AuthorityErrors->Conflicts[0];
+			const FConcertObjectInStreamID AttemptedObjectId { StreamId, ObjectReplicator->TestObject };
+			TestEqual(TEXT("AttemptedObject"), Conflict.AttemptedObject, AttemptedObjectId);
+			TestEqual(TEXT("ConflictingObject"), Conflict.ConflictingObject, { { AttemptedObjectId }, ClientConflictingWith });
+		}
+	END_DEFINE_SPEC(FPutStateRejectionSpec);
+
+	/** This tests cases of where FConcertReplication_PutState_Request should be rejected or accepted. */
+	void FPutStateRejectionSpec::Define()
+	{
+		BeforeEach([this]
+		{
+			ObjectReplicator = MakeUnique<FObjectTestReplicator>();
+			Server = MakeUnique<FReplicationServer>(*this);
+			Client1 = &Server->ConnectClient();
+			Client2 = &Server->ConnectClient();
+			Client1->JoinReplication();
+			Client2->JoinReplication();
+			
+			StreamData = ObjectReplicator->CreateStream(StreamId);
+			StreamData_FloatOnly = ObjectReplicator->CreateStreamWithProperties(StreamId, EPropertyTypeFlags::Float);
+			StreamData_VectorOnly = ObjectReplicator->CreateStreamWithProperties(StreamId, EPropertyTypeFlags::Vector);
+		});
+		AfterEach([this]
+		{
+			Server.Reset();
+			ObjectReplicator.Reset();
+		});
+
+		Describe("When an endpoint ID is invalid", [this]
+		{
+			const auto MakeRequest = [this](EConcertReplicationPutStateFlags Flags)
+			{
+				const FGuid& ClientId = Client1->GetEndpointId();
+				FConcertReplication_PutState_Request Request;
+				Request.Flags = Flags;
+				Request.NewStreams.Add(Client1->GetEndpointId(), {{ StreamData }});
+				Request.NewStreams.Add(FGuid(~ClientId.A, 0, 0, 0), {{ StreamData }});
+				return Request;
+			};
+			
+			It("EConcertReplicationChangeClientsFlags::SkipDisconnectedClients is set, the request succeeds", [this, MakeRequest]
+			{
+				IConcertClientReplicationManager& ReplicationManager = Client1->GetClientReplicationManager();
+				
+				bool bReceivedResponse = false;
+				ReplicationManager
+					.PutClientState(MakeRequest(EConcertReplicationPutStateFlags::SkipDisconnectedClients))
+					.Next([this, &bReceivedResponse](FConcertReplication_PutState_Response&& Response)
+					{
+						bReceivedResponse = true;
+						TestTrue(TEXT("Success"), Response.IsSuccess());
+					});
+				
+				TestTrue(TEXT("Received response"), bReceivedResponse);
+				TestEqual(TEXT("Stream data was applied"), ReplicationManager.GetRegisteredStreams(), { StreamData });
+			});
+			It("EConcertReplicationChangeClientsFlags::SkipDisconnectedClients not set, the request fails.", [this, MakeRequest]
+			{
+				bool bReceivedResponse = false;
+				Client1->GetClientReplicationManager()
+					.PutClientState(MakeRequest(EConcertReplicationPutStateFlags::None))
+					.Next([this, &bReceivedResponse](FConcertReplication_PutState_Response&& Response)
+					{
+						bReceivedResponse = true;
+						TestEqual(TEXT("Response code"), Response.ResponseCode, EConcertReplicationPutStateResponseCode::ClientUnknown);
+						TestEqual(TEXT("1 unknown endpoint"), Response.UnknownEndpoints.Num(), 1);
+						const FGuid ExpectedEndpoint = FGuid(~Client1->GetEndpointId().A, 0, 0, 0);
+						TestTrue(TEXT("UnknwonEndpoints contains the invalid endpoint"), Response.UnknownEndpoints.Contains(ExpectedEndpoint));
+					});
+				TestTrue(TEXT("Received response"), bReceivedResponse);
+			});
+		});
+
+		It("When creating stream with the same content for 2 clients, the request succeeds.", [this]
+		{
+			const FConcertReplication_ChangeStream_Request StreamChange { .StreamsToAdd = { StreamData } };
+			FConcertReplication_PutState_Request Request;
+			Request.NewStreams.Add(Client1->GetEndpointId(), {{ StreamData }});
+			Request.NewStreams.Add(Client2->GetEndpointId(), {{ StreamData }});
+
+			bool bReceivedResponse = false;
+			Client1->GetClientReplicationManager()
+				.PutClientState(Request)
+				.Next([this, &bReceivedResponse](FConcertReplication_PutState_Response&& Response)
+				{
+					bReceivedResponse = true;
+					TestTrue(TEXT("Success"), Response.IsSuccess());
+				});
+			TestTrue(TEXT("Received response"), bReceivedResponse);
+		});
+
+		Describe("When client 1 has authority over object", [this]
+		{
+			BeforeEach([this]
+			{
+				const FConcertReplication_ChangeStream_Request StreamChange { .StreamsToAdd = { StreamData_FloatOnly } };
+				
+				Client1->GetClientReplicationManager().ChangeStream(StreamChange);
+				Client1->GetClientReplicationManager().TakeAuthorityOver({ ObjectReplicator->TestObject });
+				
+				Client2->GetClientReplicationManager().ChangeStream(StreamChange);
+			});
+			
+			It("When request tries to give client 2 authority, the request fails", [this]
+			{
+				FConcertReplication_PutState_Request Request;
+				Request.NewAuthorityState.Add(Client2->GetEndpointId(), { { FConcertObjectInStreamID{ StreamId, ObjectReplicator->TestObject } } });
+				
+				bool bReceivedResponse = false;
+				Client1->GetClientReplicationManager().PutClientState(Request)
+					.Next([this, &bReceivedResponse](FConcertReplication_PutState_Response&& Response)
+					{
+						bReceivedResponse = true;
+						TestEqual(TEXT("Response code"), Response.ResponseCode, EConcertReplicationPutStateResponseCode::AuthorityConflict);
+						TestEqual(TEXT("AuthorityErrors.Num()"), Response.AuthorityChangeConflicts.Num(), 1);
+						TestClientHasConflictWith(Response, Client2->GetEndpointId(), Client1->GetEndpointId());
+					});
+				TestTrue(TEXT("Received response"), bReceivedResponse);
+			});
+
+			It("When request removes authority from client 1 and gives it to client 2, the request succeeds", [this]
+			{
+				FConcertReplication_PutState_Request Request;
+				Request.NewAuthorityState.Add(Client1->GetEndpointId(), {});
+				Request.NewAuthorityState.Add(Client2->GetEndpointId(), { { FConcertObjectInStreamID{ StreamId, ObjectReplicator->TestObject } } });
+				
+				bool bReceivedResponse = false;
+				Client1->GetClientReplicationManager()
+					.PutClientState(Request)
+					.Next([this, &bReceivedResponse](FConcertReplication_PutState_Response&& Response)
+					{
+						bReceivedResponse = true;
+						TestTrue(TEXT("Success"), Response.IsSuccess());
+					});
+				TestTrue(TEXT("Received response"), bReceivedResponse);
+			});
+
+			It("When request changes client 1's stream so it no longer overlaps with client 2, and gives client 2 authority, the request succeeds", [this]
+			{
+				FConcertReplication_PutState_Request Request;
+				Request.NewStreams.Add(Client1->GetEndpointId(), { .Streams { StreamData_VectorOnly } });
+				Request.NewAuthorityState.Add(Client2->GetEndpointId(), { { FConcertObjectInStreamID{ StreamId, ObjectReplicator->TestObject } } });
+
+				bool bReceivedResponse = false;
+				Client1->GetClientReplicationManager()
+					.PutClientState(Request)
+					.Next([this, &bReceivedResponse](FConcertReplication_PutState_Response&& Response)
+					{
+						bReceivedResponse = true;
+						TestTrue(TEXT("Success"), Response.IsSuccess());
+					});
+				TestTrue(TEXT("Received response"), bReceivedResponse);
+			});
+		});
+
+		It("When request tries to give two clients overlapping authority, the request fails", [this]
+		{
+			FConcertReplication_PutState_Request Request;
+			Request.NewStreams.Add(Client1->GetEndpointId(), { .Streams { StreamData } });
+			Request.NewStreams.Add(Client2->GetEndpointId(), { .Streams { StreamData } });
+			Request.NewAuthorityState.Add(Client1->GetEndpointId(), { { FConcertObjectInStreamID{ StreamId, ObjectReplicator->TestObject } } });
+			Request.NewAuthorityState.Add(Client2->GetEndpointId(), { { FConcertObjectInStreamID{ StreamId, ObjectReplicator->TestObject } } });
+
+			bool bReceivedResponse = false;
+			Client1->GetClientReplicationManager()
+				.PutClientState(Request)
+				.Next([this, &bReceivedResponse](FConcertReplication_PutState_Response&& Response)
+				{
+					bReceivedResponse = true;
+					TestEqual(TEXT("Response code"), Response.ResponseCode, EConcertReplicationPutStateResponseCode::AuthorityConflict);
+					TestEqual(TEXT("AuthorityErrors.Num()"), Response.AuthorityChangeConflicts.Num(), 2);
+					TestClientHasConflictWith(Response, Client1->GetEndpointId(), Client2->GetEndpointId());
+					TestClientHasConflictWith(Response, Client2->GetEndpointId(), Client1->GetEndpointId());
+				});
+			TestTrue(TEXT("Received response"), bReceivedResponse);
+		});
+
+		It("When request tries to mute object that is unknown, the request fails", [this]
+		{
+			FConcertReplication_PutState_Request Request;
+			Request.MuteChange.ObjectsToMute.Add(ObjectReplicator->TestObject);
+			
+			bool bReceivedResponse = false;
+			Client1->GetClientReplicationManager()
+				.PutClientState(Request)
+				.Next([this, &bReceivedResponse](FConcertReplication_PutState_Response&& Response)
+				{
+					bReceivedResponse = true;
+					TestEqual(TEXT("Response code"), Response.ResponseCode, EConcertReplicationPutStateResponseCode::MuteError);
+				});
+			TestTrue(TEXT("Received response"), bReceivedResponse);
+		});
+		It("When request tries to mute object that will become unknown, the request fails", [this]
+		{
+			IConcertClientReplicationManager& ReplicationManager = Client1->GetClientReplicationManager();
+			ReplicationManager.ChangeStream({ .StreamsToAdd = { StreamData } });
+			
+			FConcertReplication_PutState_Request Request;
+			Request.NewStreams.Add(Client1->GetEndpointId());
+			Request.MuteChange.ObjectsToMute.Add(ObjectReplicator->TestObject);
+			bool bReceivedResponse = false;
+			ReplicationManager.PutClientState(Request)
+				.Next([this, &bReceivedResponse](FConcertReplication_PutState_Response&& Response)
+				{
+					bReceivedResponse = true;
+					TestEqual(TEXT("Response code"), Response.ResponseCode, EConcertReplicationPutStateResponseCode::MuteError);
+				});
+			TestTrue(TEXT("Received response"), bReceivedResponse);
+		});
+		It("When request tries to mute an object that will become known, the request succeeds", [this]
+		{
+			IConcertClientReplicationManager& ReplicationManager = Client1->GetClientReplicationManager();
+			
+			FConcertReplication_PutState_Request Request;
+			Request.NewStreams.Add(Client1->GetEndpointId(), {{ StreamData }});
+			Request.MuteChange.ObjectsToMute.Add(ObjectReplicator->TestObject);
+			bool bReceivedResponse = false;
+			ReplicationManager.PutClientState(Request)
+				.Next([this, &bReceivedResponse](FConcertReplication_PutState_Response&& Response)
+				{
+					bReceivedResponse = true;
+					TestTrue(TEXT("Success"), Response.IsSuccess());
+				});
+			TestTrue(TEXT("Received response"), bReceivedResponse);
+		});
+	}
+}
diff --git a/Engine/Plugins/Developer/Concert/ConcertSync/ConcertSyncTest/Source/ConcertSyncTest/Private/Replication/PutState/PutStateSyncControl.spec.cpp b/Engine/Plugins/Developer/Concert/ConcertSync/ConcertSyncTest/Source/ConcertSyncTest/Private/Replication/PutState/PutStateSyncControl.spec.cpp
new file mode 100644
--- /dev/null
+++ b/Engine/Plugins/Developer/Concert/ConcertSync/ConcertSyncTest/Source/ConcertSyncTest/Private/Replication/PutState/PutStateSyncControl.spec.cpp
@@ -0,0 +1,184 @@
+﻿// Copyright Epic Games, Inc. All Rights Reserved.
+
+#include "Replication/Util/ConcertClientReplicationBridgeMock.h"
+#include "Replication/Util/Spec/ObjectTestReplicator.h"
+#include "Replication/Util/Spec/ReplicationClient.h"
+#include "Replication/Util/Spec/ReplicationServer.h"
+
+#include "Misc/AutomationTest.h"
+
+namespace UE::ConcertSyncTests::Replication::ChangeClients
+{
+	BEGIN_DEFINE_SPEC(FPutStateSyncControlSpec, "Editor.Concert.Replication.PutState.SyncControl", EAutomationTestFlags::EditorContext | EAutomationTestFlags::EngineFilter)
+		TUniquePtr<FObjectTestReplicator> ObjectReplicator;
+
+		TUniquePtr<FReplicationServer> Server;
+		FReplicationClient* Sender = nullptr;
+		FReplicationClient* Receiver = nullptr;
+
+		const FGuid StreamId = FGuid::NewGuid();
+		FConcertReplicationStream StreamData;
+	END_DEFINE_SPEC(FPutStateSyncControlSpec);
+
+	/**
+	 * This tests that sync control is correct after a FConcertReplication_PutState change.
+	 */
+	void FPutStateSyncControlSpec::Define()
+	{
+		BeforeEach([this]
+		{
+			ObjectReplicator = MakeUnique<FObjectTestReplicator>();
+			Server = MakeUnique<FReplicationServer>(*this);
+			Sender = &Server->ConnectClient();
+			Receiver = &Server->ConnectClient();
+
+			StreamData = ObjectReplicator->CreateStream(StreamId);
+			Sender->JoinReplication(ObjectReplicator->CreateSenderArgs(StreamId));
+			Receiver->JoinReplicationAsListener({ ObjectReplicator->TestObject });
+			
+			IConcertClientReplicationManager& ReplicationManager = Sender->GetClientReplicationManager();
+			ReplicationManager.TakeAuthorityOver({ ObjectReplicator->TestObject });
+			
+			// Make sure that we've set up the test correctly: Sender should now have sync control.
+			const bool bHasSyncControl = ReplicationManager.GetSyncControlledObjects().Contains({ StreamId, ObjectReplicator->TestObject });
+			TestTrue(TEXT("bHasSyncControl"), bHasSyncControl);
+		});
+		AfterEach([this]
+		{
+			Server.Reset();
+			ObjectReplicator.Reset();
+		});
+
+		Describe("When client puts empty state", [this]
+		{
+			const auto RunTest = [this]()
+			{
+				IConcertClientReplicationManager& ReplicationManager = Sender->GetClientReplicationManager();
+			
+				const FConcertReplication_PutState_Request Request { .NewStreams = { { Sender->GetEndpointId(), {} }} };
+				bool bReceivedResponse = false;
+				TFuture<FConcertReplication_PutState_Response> Future = ReplicationManager
+					.PutClientState(Request)
+					.Next([this, &bReceivedResponse](FConcertReplication_PutState_Response&& Response)
+					{
+						bReceivedResponse = true;
+						// The change should not show up in the response because the client should predict that it does not have any more sync control.
+						return Response;
+					});
+				TestTrue(TEXT("bReceivedResponse"), bReceivedResponse);
+				return Future;
+			};
+		
+			It("Response contains no sync control change", [this, RunTest]()
+			{
+				RunTest()
+					.Next([this](FConcertReplication_PutState_Response&& Response)
+					{
+						TestEqual(TEXT("NewControlStates.Num()"), Response.SyncControl.NewControlStates.Num(), 0);
+					});
+			});
+			It("Local client cache thinks it has no sync control", [this, RunTest]
+			{
+				RunTest();
+				TestEqual(TEXT("No sync control"), Sender->GetClientReplicationManager().GetSyncControlledObjects().Num(), 0);
+			});
+		});
+
+		Describe("When client's put request retains sync control", [this]
+		{
+			const auto RunTest = [this](bool bTakeAuthority)
+			{
+				IConcertClientReplicationManager& ReplicationManager = Sender->GetClientReplicationManager();
+			
+				FConcertReplication_PutState_Request Request;
+				Request.NewStreams.Add(Sender->GetEndpointId()).Streams.Add(ObjectReplicator->CreateStream(StreamId));
+				if (bTakeAuthority)
+				{
+					Request.NewAuthorityState.Add(Sender->GetEndpointId()).Objects.Add({ StreamId, ObjectReplicator->TestObject });
+				}
+				bool bReceivedResponse = false;
+				TFuture<FConcertReplication_PutState_Response> Future = ReplicationManager
+					.PutClientState(Request)
+					.Next([this, &bReceivedResponse](FConcertReplication_PutState_Response&& Response)
+					{
+						bReceivedResponse = true;
+						return Response;
+					});
+				TestTrue(TEXT("bReceivedResponse"), bReceivedResponse);
+				return Future;
+			};
+		
+			It("Response contains no sync control change (specify authority: true)", [this, RunTest]()
+			{
+				RunTest(true)
+					.Next([this](FConcertReplication_PutState_Response&& Response)
+					{
+						// The client is expected tp predict that it no longer has sync control.
+						TestEqual(TEXT("NewControlStates.Num()"), Response.SyncControl.NewControlStates.Num(), 0);
+					});
+			});
+			It("Response contains no sync control change (specify authority: false)", [this, RunTest]()
+			{
+				RunTest(true)
+					.Next([this](FConcertReplication_PutState_Response&& Response)
+					{
+						// The client is expected tp predict that it no longer has sync control.
+						TestEqual(TEXT("NewControlStates.Num()"), Response.SyncControl.NewControlStates.Num(), 0);
+					});
+			});
+			
+			It("Local client cache thinks it has sync control (specify authority: true)", [this, RunTest]
+			{
+				RunTest(true);
+				
+				// The client should have predicted that it no longer has sync control.
+				TestEqual(TEXT("GetSyncControlledObjects().Num()"), Sender->GetClientReplicationManager().GetSyncControlledObjects().Num(), 1);
+				TestTrue(TEXT("Has Sync Control"), Sender->GetClientReplicationManager().GetSyncControlledObjects().Contains({ StreamId, ObjectReplicator->TestObject }));
+			});
+			It("Local client cache thinks it has sync control (specify authority: false)", [this, RunTest]
+			{
+				RunTest(false);
+				
+				// The client should have predicted that it no longer has sync control.
+				TestEqual(TEXT("GetSyncControlledObjects().Num()"), Sender->GetClientReplicationManager().GetSyncControlledObjects().Num(), 1);
+				TestTrue(TEXT("Has Sync Control"), Sender->GetClientReplicationManager().GetSyncControlledObjects().Contains({ StreamId, ObjectReplicator->TestObject }));
+			});
+		});
+
+		Describe("When client has sync control and removes object", [this]
+		{
+			const auto RunTest = [this]()
+			{
+				IConcertClientReplicationManager& ReplicationManager = Sender->GetClientReplicationManager();
+				const FConcertReplication_PutState_Request Request { .NewStreams = {  { Sender->GetEndpointId(), {}} } };
+				bool bReceivedResponse = false;
+				TFuture<FConcertReplication_PutState_Response> Future = ReplicationManager
+					.PutClientState(Request)
+					.Next([this, &bReceivedResponse](FConcertReplication_PutState_Response&& Response)
+					{
+						bReceivedResponse = true;
+						return Response;
+					});
+				TestTrue(TEXT("bReceivedResponse"), bReceivedResponse);
+				return Future;
+			};
+
+			It("Response contains no sync control change", [this, RunTest]
+			{
+				RunTest()
+					.Next([this](FConcertReplication_PutState_Response&& Response)
+					{
+						// The client is expected tp predict that it no longer has sync control.
+						TestEqual(TEXT("NewControlStates.Num()"), Response.SyncControl.NewControlStates.Num(), 0);
+					});
+			});
+			It("Local client cache does not think it has sync control", [this, RunTest]
+			{
+				RunTest();
+				
+				// The client should have predicted that it no longer has sync control.
+				TestEqual(TEXT("GetSyncControlledObjects().Num()"), Sender->GetClientReplicationManager().GetSyncControlledObjects().Num(), 0);
+			});
+		});
+	}
+}
diff --git a/Engine/Plugins/Developer/Concert/ConcertSync/ConcertSyncTest/Source/ConcertSyncTest/Private/Replication/SyncControl.spec.cpp b/Engine/Plugins/Developer/Concert/ConcertSync/ConcertSyncTest/Source/ConcertSyncTest/Private/Replication/SyncControl.spec.cpp
--- a/Engine/Plugins/Developer/Concert/ConcertSync/ConcertSyncTest/Source/ConcertSyncTest/Private/Replication/SyncControl.spec.cpp
+++ b/Engine/Plugins/Developer/Concert/ConcertSync/ConcertSyncTest/Source/ConcertSyncTest/Private/Replication/SyncControl.spec.cpp
@@ -289,7 +289,7 @@ namespace UE::ConcertSyncTests::Replication::UI
 					.Next([this, &Manager, &ReleaseAuthority, &SyncControlBeforeRelease](FConcertReplication_ChangeAuthority_Response&& Response)
 					{
 						// Aggregate is validated in SyncControlState.spec.cpp
-						SyncControlBeforeRelease.AppendAuthorityChange(ReleaseAuthority, Response);
+						SyncControlBeforeRelease.AppendAuthorityChange(ReleaseAuthority, Response.SyncControl);
 						TestTrue(TEXT("Client predicted correctly"), SyncControlBeforeRelease == Manager.GetSyncControlledObjects());
 					});
 
@@ -300,7 +300,7 @@ namespace UE::ConcertSyncTests::Replication::UI
 					.Next([this, &Manager, &SyncControlBeforeTake, &TakeAuthority](FConcertReplication_ChangeAuthority_Response&& Response)
 					{
 						// Aggregate is validated in SyncControlState.spec.cpp
-						SyncControlBeforeTake.AppendAuthorityChange(TakeAuthority, Response);
+						SyncControlBeforeTake.AppendAuthorityChange(TakeAuthority, Response.SyncControl);
 						TestTrue(TEXT("Client predicted correctly"), SyncControlBeforeTake == Manager.GetSyncControlledObjects());
 					});
 			});
diff --git a/Engine/Plugins/Developer/Concert/ConcertSync/ConcertSyncTest/Source/ConcertSyncTest/Private/Replication/Util/Spec/ObjectTestReplicator.cpp b/Engine/Plugins/Developer/Concert/ConcertSync/ConcertSyncTest/Source/ConcertSyncTest/Private/Replication/Util/Spec/ObjectTestReplicator.cpp
--- a/Engine/Plugins/Developer/Concert/ConcertSync/ConcertSyncTest/Source/ConcertSyncTest/Private/Replication/Util/Spec/ObjectTestReplicator.cpp
+++ b/Engine/Plugins/Developer/Concert/ConcertSync/ConcertSyncTest/Source/ConcertSyncTest/Private/Replication/Util/Spec/ObjectTestReplicator.cpp
@@ -51,6 +51,15 @@ namespace UE::ConcertSyncTests::Replication
 		return SendingStream;
 	}
 
+	FConcertReplicationStream FObjectTestReplicator::CreateStreamWithProperties(FGuid SenderStreamId, EPropertyTypeFlags PropertyTypeFlags, EConcertObjectReplicationMode ReplicationMode, uint8 ReplicationRate) const
+	{
+		FConcertReplicationStream SendingStream;
+		SendingStream.BaseDescription.Identifier = SenderStreamId;
+		SendingStream.BaseDescription.FrequencySettings.Defaults = { ReplicationMode, ReplicationRate };
+		AddToStreamWithProperties(SendingStream, PropertyTypeFlags);
+		return SendingStream;
+	}
+
 	void FObjectTestReplicator::AddToStream(FConcertReplicationStream& Stream, EConcertObjectReplicationMode ReplicationMode, uint8 ReplicationRate) const
 	{
 		FConcertReplicatedObjectInfo ReplicatedObjectInfo { TestObject->GetClass() };
@@ -68,6 +77,40 @@ namespace UE::ConcertSyncTests::Replication
 		}
 	}
 
+	void FObjectTestReplicator::AddToStreamWithProperties(FConcertReplicationStream& Stream, EPropertyTypeFlags PropertyTypeFlags, EConcertObjectReplicationMode ReplicationMode, uint8 ReplicationRate) const
+	{
+		FConcertReplicatedObjectInfo ReplicatedObjectInfo { TestObject->GetClass() };
+		ConcertSyncCore::PropertyChain::ForEachReplicatableConcertProperty(*TestObject->GetClass(), [&ReplicatedObjectInfo, PropertyTypeFlags](FConcertPropertyChain&& Chain)
+		{
+			const bool bIsFloat = Chain == TArray<FName>{ TEXT("Float") };
+			const bool bIsVector = Chain.GetRootProperty() == TEXT("Vector");
+			if (!EnumHasAnyFlags(PropertyTypeFlags, EPropertyTypeFlags::Float)
+				&& bIsFloat)
+			{
+				return EBreakBehavior::Continue;
+			}
+			if (!EnumHasAnyFlags(PropertyTypeFlags, EPropertyTypeFlags::Vector)
+				&& bIsVector)
+			{
+				return EBreakBehavior::Continue;
+			}
+			if (!EnumHasAnyFlags(PropertyTypeFlags, EPropertyTypeFlags::Others) && !bIsFloat && !bIsVector)
+			{
+				return EBreakBehavior::Continue;
+			}
+			
+			ReplicatedObjectInfo.PropertySelection.ReplicatedProperties.Emplace(MoveTemp(Chain));
+			return EBreakBehavior::Continue;
+		});
+		Stream.BaseDescription.ReplicationMap.ReplicatedObjects.Add(TestObject, ReplicatedObjectInfo);
+
+		const FConcertObjectReplicationSettings FrequencySettings { ReplicationMode, ReplicationRate };
+		if (Stream.BaseDescription.FrequencySettings.Defaults != FrequencySettings)
+		{
+			Stream.BaseDescription.FrequencySettings.ObjectOverrides.Add(TestObject, FrequencySettings);
+		}
+	}
+
 	void FObjectTestReplicator::SimulateSendObjectToReceiver(
 		FAutomationTestBase& Test,
 		FObjectReplicationContext Context,
diff --git a/Engine/Plugins/Developer/Concert/ConcertSync/ConcertSyncTest/Source/ConcertSyncTest/Private/Replication/Util/Spec/ObjectTestReplicator.h b/Engine/Plugins/Developer/Concert/ConcertSync/ConcertSyncTest/Source/ConcertSyncTest/Private/Replication/Util/Spec/ObjectTestReplicator.h
--- a/Engine/Plugins/Developer/Concert/ConcertSync/ConcertSyncTest/Source/ConcertSyncTest/Private/Replication/Util/Spec/ObjectTestReplicator.h
+++ b/Engine/Plugins/Developer/Concert/ConcertSync/ConcertSyncTest/Source/ConcertSyncTest/Private/Replication/Util/Spec/ObjectTestReplicator.h
@@ -32,6 +32,17 @@ namespace UE::ConcertSyncTests::Replication
 	};
 	ENUM_CLASS_FLAGS(EPropertyReplicationFlags);
 
+	/** Defines the properties you want to add to the stream */
+	enum class EPropertyTypeFlags : uint8
+	{
+		None,
+		Float = 1 << 0,
+		Vector = 1 << 1,
+		Others = 1 << 2,
+		All = Float | Vector | Others
+	};
+	ENUM_CLASS_FLAGS(EPropertyTypeFlags);
+
 	struct FObjectReplicationContext
 	{
 		FReplicationClient& Sender;
@@ -74,18 +85,34 @@ namespace UE::ConcertSyncTests::Replication
 			EConcertObjectReplicationMode ReplicationMode = EConcertObjectReplicationMode::Realtime,
 			uint8 ReplicationRate = 30
 			) const;
+		
 		/** Util for creating a stream that replicates TestObject. */
 		FConcertReplicationStream CreateStream(
 			FGuid SenderStreamId = FGuid::NewGuid(),
 			EConcertObjectReplicationMode ReplicationMode = EConcertObjectReplicationMode::Realtime,
 			uint8 ReplicationRate = 30
 		) const;
+		/** Util for creating a stream that replicates TestObject and specifies the properties. */
+		FConcertReplicationStream CreateStreamWithProperties(
+			FGuid SenderStreamId = FGuid::NewGuid(),
+			EPropertyTypeFlags PropertyTypeFlags = EPropertyTypeFlags::All,
+			EConcertObjectReplicationMode ReplicationMode = EConcertObjectReplicationMode::Realtime,
+			uint8 ReplicationRate = 30
+		) const;
+		
 		/** Util adding replication settings for TestObject to a stream */
 		void AddToStream(
 			FConcertReplicationStream& Stream,
 			EConcertObjectReplicationMode ReplicationMode = EConcertObjectReplicationMode::Realtime,
 			uint8 ReplicationRate = 30
 			) const;
+		/** Util adding replication settings for TestObject to a stream and specifies the properties. */
+		void AddToStreamWithProperties(
+			FConcertReplicationStream& Stream,
+			EPropertyTypeFlags PropertyTypeFlags = EPropertyTypeFlags::All,
+			EConcertObjectReplicationMode ReplicationMode = EConcertObjectReplicationMode::Realtime,
+			uint8 ReplicationRate = 30
+			) const;
 		
 		/** This overload allows you to send properties but from multiple streams. */
 		void SimulateSendObjectToReceiver(
