Fix typo (identical sub-expressions to the left and to the right of the '==' operator)

diff --git a/Engine/Plugins/Developer/Concert/ConcertSync/ConcertSyncServer/Source/ConcertSyncServer/Private/Replication/AuthorityManager.cpp b/Engine/Plugins/Developer/Concert/ConcertSync/ConcertSyncServer/Source/ConcertSyncServer/Private/Replication/AuthorityManager.cpp
--- a/Engine/Plugins/Developer/Concert/ConcertSync/ConcertSyncServer/Source/ConcertSyncServer/Private/Replication/AuthorityManager.cpp
+++ b/Engine/Plugins/Developer/Concert/ConcertSync/ConcertSyncServer/Source/ConcertSyncServer/Private/Replication/AuthorityManager.cpp
@@ -82,7 +82,7 @@ namespace UE::ConcertSyncServer::Replication
 			const FConcertObjectReplicationMap* ReplicationMap = nullptr;
 			Getters.ForEachStream(ClientId, [&StreamId, &ReplicationMap](const FGuid& InStreamId, const FConcertObjectReplicationMap& InReplicationMap) mutable
 			{
-				if (InStreamId == InStreamId)
+				if (InStreamId == StreamId)
 				{
 					ReplicationMap = &InReplicationMap;
 					return EBreakBehavior::Break;
