Improve naming for better self documentation

diff --git a/Engine/Plugins/Developer/Concert/ConcertSync/ConcertSyncServer/Source/ConcertSyncServer/Private/Replication/Util/GroundTruthOverride.cpp b/Engine/Plugins/Developer/Concert/ConcertSync/ConcertSyncServer/Source/ConcertSyncServer/Private/Replication/Util/GroundTruthOverride.cpp
--- a/Engine/Plugins/Developer/Concert/ConcertSync/ConcertSyncServer/Source/ConcertSyncServer/Private/Replication/Util/GroundTruthOverride.cpp
+++ b/Engine/Plugins/Developer/Concert/ConcertSync/ConcertSyncServer/Source/ConcertSyncServer/Private/Replication/Util/GroundTruthOverride.cpp
@@ -23,7 +23,7 @@ namespace UE::ConcertSyncServer::Replication
 		}
 		else
 		{
-			Clients.ForEachStream(ClientEndpointId, [&Callback](const FConcertReplicationStream& Stream)
+			NoOverrideStreamFallback.ForEachStream(ClientEndpointId, [&Callback](const FConcertReplicationStream& Stream)
 			{
 				return Callback(Stream.BaseDescription.Identifier, Stream.BaseDescription.ReplicationMap);
 			});
@@ -32,7 +32,7 @@ namespace UE::ConcertSyncServer::Replication
 
 	void FGroundTruthOverride::ForEachSendingClient(TFunctionRef<EBreakBehavior(const FGuid& ClientEndpointId)> Callback) const
 	{
-		Clients.ForEachReplicationClient(Callback);
+		NoOverrideStreamFallback.ForEachReplicationClient(Callback);
 	}
 
 	bool FGroundTruthOverride::HasAuthority(const FGuid& ClientId, const FGuid& StreamId, const FSoftObjectPath& ObjectPath) const
@@ -41,6 +41,6 @@ namespace UE::ConcertSyncServer::Replication
 		const FConcertObjectInStreamArray* AuthorityOverride = AuthorityOverrides.Find(ClientId);
 		return AuthorityOverride
 			? AuthorityOverride->Objects.Contains(ObjectId)
-			: AuthorityManager.HasAuthorityToChange(ObjectId);
+			: NoOverrideAuthorityFallback.HasAuthorityToChange(ObjectId);
 	}
 }
diff --git a/Engine/Plugins/Developer/Concert/ConcertSync/ConcertSyncServer/Source/ConcertSyncServer/Private/Replication/Util/GroundTruthOverride.h b/Engine/Plugins/Developer/Concert/ConcertSync/ConcertSyncServer/Source/ConcertSyncServer/Private/Replication/Util/GroundTruthOverride.h
--- a/Engine/Plugins/Developer/Concert/ConcertSync/ConcertSyncServer/Source/ConcertSyncServer/Private/Replication/Util/GroundTruthOverride.h
+++ b/Engine/Plugins/Developer/Concert/ConcertSync/ConcertSyncServer/Source/ConcertSyncServer/Private/Replication/Util/GroundTruthOverride.h
@@ -32,15 +32,14 @@ namespace UE::ConcertSyncServer::Replication
 		FGroundTruthOverride(
 			const TMap<FGuid, FConcertReplicationStreamArray>& StreamOverrides UE_LIFETIMEBOUND,
 			const TMap<FGuid, FConcertObjectInStreamArray>& AuthorityOverrides UE_LIFETIMEBOUND,
-			const IRegistrationEnumerator& Clients UE_LIFETIMEBOUND,
-			const FAuthorityManager& AuthorityManager UE_LIFETIMEBOUND
+			const IRegistrationEnumerator& NoOverrideStreamFallback UE_LIFETIMEBOUND,
+			const FAuthorityManager& NoOverrideAuthorityFallback UE_LIFETIMEBOUND
 			)
 			: StreamOverrides(StreamOverrides)
 			, AuthorityOverrides(AuthorityOverrides)
-			, Clients(Clients)
-			, AuthorityManager(AuthorityManager)
-		{
-		}
+			, NoOverrideStreamFallback(NoOverrideStreamFallback)
+			, NoOverrideAuthorityFallback(NoOverrideAuthorityFallback)
+		{}
 
 		//~ Begin IReplicationGroundTruth Interface
 		virtual void ForEachStream(const FGuid& ClientEndpointId, TFunctionRef<EBreakBehavior(const FGuid& StreamId, const FConcertObjectReplicationMap& ReplicationMap)> Callback) const override;
@@ -53,8 +52,10 @@ namespace UE::ConcertSyncServer::Replication
 		const TMap<FGuid, FConcertReplicationStreamArray>& StreamOverrides;
 		const TMap<FGuid, FConcertObjectInStreamArray>& AuthorityOverrides;
 
-		const IRegistrationEnumerator& Clients;
-		const FAuthorityManager& AuthorityManager;
+		/** Gives us the stream content when no override was specified. */
+		const IRegistrationEnumerator& NoOverrideStreamFallback;
+		/** Gives us the authority when no override was specified. */
+		const FAuthorityManager& NoOverrideAuthorityFallback;
 	};
 }
 
