Extends FGroundTruthOverride so it reports additional clients that are not actually present.
Also renames IReplicationGroundTruth::ForEachSendingClient to ForEachClient to not falsely give readers the impression that the enumerated list is filtered by whether or not the client is sending.

diff --git a/Engine/Plugins/Developer/Concert/ConcertApp/MultiUserClient/Source/MultiUserClient/Private/Replication/Misc/GlobalAuthorityCache.cpp b/Engine/Plugins/Developer/Concert/ConcertApp/MultiUserClient/Source/MultiUserClient/Private/Replication/Misc/GlobalAuthorityCache.cpp
--- a/Engine/Plugins/Developer/Concert/ConcertApp/MultiUserClient/Source/MultiUserClient/Private/Replication/Misc/GlobalAuthorityCache.cpp
+++ b/Engine/Plugins/Developer/Concert/ConcertApp/MultiUserClient/Source/MultiUserClient/Private/Replication/Misc/GlobalAuthorityCache.cpp
@@ -278,7 +278,7 @@ namespace UE::MultiUserClient
 		Callback(StreamSynchronizer.GetStreamId(), StreamSynchronizer.GetServerState());
 	}
 
-	void FGlobalAuthorityCache::ForEachSendingClient(TFunctionRef<EBreakBehavior(const FGuid& ClientEndpointId)> Callback) const
+	void FGlobalAuthorityCache::ForEachClient(TFunctionRef<EBreakBehavior(const FGuid& ClientEndpointId)> Callback) const
 	{
 		auto ProcessClient = [&Callback](const FReplicationClient& Client)
 		{
diff --git a/Engine/Plugins/Developer/Concert/ConcertApp/MultiUserClient/Source/MultiUserClient/Private/Replication/Misc/GlobalAuthorityCache.h b/Engine/Plugins/Developer/Concert/ConcertApp/MultiUserClient/Source/MultiUserClient/Private/Replication/Misc/GlobalAuthorityCache.h
--- a/Engine/Plugins/Developer/Concert/ConcertApp/MultiUserClient/Source/MultiUserClient/Private/Replication/Misc/GlobalAuthorityCache.h
+++ b/Engine/Plugins/Developer/Concert/ConcertApp/MultiUserClient/Source/MultiUserClient/Private/Replication/Misc/GlobalAuthorityCache.h
@@ -128,7 +128,7 @@ namespace UE::MultiUserClient
 
 		//~ Begin IReplicationGroundTruth Interface
 		virtual void ForEachStream(const FGuid& ClientEndpointId, TFunctionRef<EBreakBehavior(const FGuid& StreamId, const FConcertObjectReplicationMap& ReplicationMap)> Callback) const override;
-		virtual void ForEachSendingClient(TFunctionRef<EBreakBehavior(const FGuid& ClientEndpointId)> Callback) const override;
+		virtual void ForEachClient(TFunctionRef<EBreakBehavior(const FGuid& ClientEndpointId)> Callback) const override;
 		virtual bool HasAuthority(const FGuid& ClientId, const FGuid& StreamId, const FSoftObjectPath& ObjectPath) const override;
 		//~ End IReplicationGroundTruth Interface
 	};
diff --git a/Engine/Plugins/Developer/Concert/ConcertSync/ConcertSyncCore/Source/ConcertSyncCore/Private/Replication/AuthorityConflictSharedUtils.cpp b/Engine/Plugins/Developer/Concert/ConcertSync/ConcertSyncCore/Source/ConcertSyncCore/Private/Replication/AuthorityConflictSharedUtils.cpp
--- a/Engine/Plugins/Developer/Concert/ConcertSync/ConcertSyncCore/Source/ConcertSyncCore/Private/Replication/AuthorityConflictSharedUtils.cpp
+++ b/Engine/Plugins/Developer/Concert/ConcertSync/ConcertSyncCore/Source/ConcertSyncCore/Private/Replication/AuthorityConflictSharedUtils.cpp
@@ -17,7 +17,7 @@ namespace UE::ConcertSyncCore::Replication::AuthorityConflictUtils
             const IReplicationGroundTruth& GroundTruth
             )
 		{
-			GroundTruth.ForEachSendingClient([&Object, &Callback, &IgnoredClients, &GroundTruth](const FGuid& ClientEndpointId)
+			GroundTruth.ForEachClient([&Object, &Callback, &IgnoredClients, &GroundTruth](const FGuid& ClientEndpointId)
 			{
 				if (IgnoredClients.Contains(ClientEndpointId))
 				{
diff --git a/Engine/Plugins/Developer/Concert/ConcertSync/ConcertSyncCore/Source/ConcertSyncCore/Public/Replication/AuthorityConflictSharedUtils.h b/Engine/Plugins/Developer/Concert/ConcertSync/ConcertSyncCore/Source/ConcertSyncCore/Public/Replication/AuthorityConflictSharedUtils.h
--- a/Engine/Plugins/Developer/Concert/ConcertSync/ConcertSyncCore/Source/ConcertSyncCore/Public/Replication/AuthorityConflictSharedUtils.h
+++ b/Engine/Plugins/Developer/Concert/ConcertSync/ConcertSyncCore/Source/ConcertSyncCore/Public/Replication/AuthorityConflictSharedUtils.h
@@ -36,7 +36,7 @@ namespace UE::ConcertSyncCore::Replication::AuthorityConflictUtils
 		virtual void ForEachStream(const FGuid& ClientEndpointId, TFunctionRef<EBreakBehavior(const FGuid& StreamId, const FConcertObjectReplicationMap& ReplicationMap)> Callback) const = 0;
 
 		/** Iterates through all clients have registered to send any data. */
-		virtual void ForEachSendingClient(TFunctionRef<EBreakBehavior(const FGuid& ClientEndpointId)> Callback) const = 0;
+		virtual void ForEachClient(TFunctionRef<EBreakBehavior(const FGuid& ClientEndpointId)> Callback) const = 0;
 
 		/** @return Whether ClientId's stream StreamId has authority over ObjectPath. */
 		virtual bool HasAuthority(const FGuid& ClientId, const FGuid& StreamId, const FSoftObjectPath& ObjectPath) const = 0; 
diff --git a/Engine/Plugins/Developer/Concert/ConcertSync/ConcertSyncServer/Source/ConcertSyncServer/Private/Replication/AuthorityManager.cpp b/Engine/Plugins/Developer/Concert/ConcertSync/ConcertSyncServer/Source/ConcertSyncServer/Private/Replication/AuthorityManager.cpp
--- a/Engine/Plugins/Developer/Concert/ConcertSync/ConcertSyncServer/Source/ConcertSyncServer/Private/Replication/AuthorityManager.cpp
+++ b/Engine/Plugins/Developer/Concert/ConcertSync/ConcertSyncServer/Source/ConcertSyncServer/Private/Replication/AuthorityManager.cpp
@@ -33,7 +33,7 @@ namespace UE::ConcertSyncServer::Replication
 				});
 			}
 			
-			virtual void ForEachSendingClient(TFunctionRef<EBreakBehavior(const FGuid& ClientEndpointId)> Callback) const override
+			virtual void ForEachClient(TFunctionRef<EBreakBehavior(const FGuid& ClientEndpointId)> Callback) const override
 			{
 				Getters.ForEachReplicationClient(Callback);
 			}
@@ -155,12 +155,12 @@ namespace UE::ConcertSyncServer::Replication
 
 	FAuthorityManager::EAuthorityResult FAuthorityManager::EnumerateAuthorityConflicts(
 		const FConcertReplicatedObjectId& Object, 
-		const FConcertPropertySelection* OverwriteProperties,
+		const FConcertPropertySelection* OverrideProperties,
 		FProcessAuthorityConflict ProcessConflict
 		) const
 	{
 		const FClientId& ClientId = Object.SenderEndpointId;
-		const FConcertPropertySelection* PropertiesToCheck = OverwriteProperties;
+		const FConcertPropertySelection* PropertiesToCheck = OverrideProperties;
 		if (!PropertiesToCheck)
 		{
 			const FConcertReplicationStream* Description = Private::FindClientStreamById(Getters, ClientId, Object.StreamId);
diff --git a/Engine/Plugins/Developer/Concert/ConcertSync/ConcertSyncServer/Source/ConcertSyncServer/Private/Replication/AuthorityManager.h b/Engine/Plugins/Developer/Concert/ConcertSync/ConcertSyncServer/Source/ConcertSyncServer/Private/Replication/AuthorityManager.h
--- a/Engine/Plugins/Developer/Concert/ConcertSync/ConcertSyncServer/Source/ConcertSyncServer/Private/Replication/AuthorityManager.h
+++ b/Engine/Plugins/Developer/Concert/ConcertSync/ConcertSyncServer/Source/ConcertSyncServer/Private/Replication/AuthorityManager.h
@@ -67,29 +67,37 @@ namespace UE::ConcertSyncServer::Replication
 		/**
 		 * Enumerates all authority conflicts, if any, that would occur if ObjectChange.SenderEndpointId were to take authority over the identified object.
 		 * You can optionally supply OverwriteProperties, which is useful if you're about to change the contents of the stream.
+		 * 
+		 * @param Object Identifies the object, its stream and client that will send the object. Conflicts are checked against this.
+		 * @param OverrideProperties Optional properties that Object will have and conflicts should be checked against. If not overriden, falls back to the properties registered thus far.
+		 * @param ProcessConflict Invoked for the potential conflicts
 		 * @return Whether there were any conflicts
 		 */
 		EAuthorityResult EnumerateAuthorityConflicts(
 			const FConcertReplicatedObjectId& Object,
-			const FConcertPropertySelection* OverwriteProperties = nullptr,
-			FProcessAuthorityConflict ProcessConflict = [](auto&, auto&, auto&){ return EBreakBehavior::Break; }
+			const FConcertPropertySelection* OverrideProperties = nullptr,
+			FProcessAuthorityConflict ProcessConflict = [](const FClientId&, const FStreamId&, const FConcertPropertyChain&){ return EBreakBehavior::Break; }
 			) const;
 		/** Whether it is legal for this the client identified by ClientId to take control over the object given the stream the client has registered. */
 		bool CanTakeAuthority(const FConcertReplicatedObjectId& Object) const;
 
 		/**
 		 * Enumerates all authority conflicts, if any, that would occur if ObjectChange.SenderEndpointId were to take authority over the identified object.
-		 *
-		 * While EnumerateAuthorityConflicts allows you to specify OverwriteProperties, this version allows you to override the
+		 * This version allows you to override the entire session content (compared to EnumerateAuthorityConflicts, which only allows you to specify OverwriteProperties). 
+		 * 
+		 * @param Object Identifies the object, its stream and client that will send the object. Conflicts are checked against this.
+		 * @param StreamOverrides Overrides that the stream content of connected clients. If a client is not overriden, falls back to the content registered thus far.
+		 *		The key can be non-existing clients if you want to validate no conflicts were to occur if some other clients were present("injection").
+		 * @param AuthorityOverrides Overrides the authority of connected and "injected" clients. If a client is not overriden, falls back to the authority registered thus far.
+		 * @param ProcessConflict Invoked for the potential conflicts
 		 * 
 		 * @return Whether there were any conflicts
 		 */
-
 		EAuthorityResult EnumerateAuthorityConflictsWithOverrides(
 			const FConcertReplicatedObjectId& Object,
 			const TMap<FGuid, FConcertReplicationStreamArray>& StreamOverrides,
 			const TMap<FGuid, FConcertObjectInStreamArray>& AuthorityOverrides,
-			FProcessAuthorityConflict ProcessConflict = [](auto&, auto&, auto&){ return EBreakBehavior::Break; }
+			FProcessAuthorityConflict ProcessConflict = [](const FClientId&, const FStreamId&, const FConcertPropertyChain&){ return EBreakBehavior::Break; }
 			) const;
 
 		/** Notifies this manager that the client has left, which means all their authority is now gone. */
diff --git a/Engine/Plugins/Developer/Concert/ConcertSync/ConcertSyncServer/Source/ConcertSyncServer/Private/Replication/ConcertServerReplicationManager_PutState.cpp b/Engine/Plugins/Developer/Concert/ConcertSync/ConcertSyncServer/Source/ConcertSyncServer/Private/Replication/ConcertServerReplicationManager_PutState.cpp
--- a/Engine/Plugins/Developer/Concert/ConcertSync/ConcertSyncServer/Source/ConcertSyncServer/Private/Replication/ConcertServerReplicationManager_PutState.cpp
+++ b/Engine/Plugins/Developer/Concert/ConcertSync/ConcertSyncServer/Source/ConcertSyncServer/Private/Replication/ConcertServerReplicationManager_PutState.cpp
@@ -190,6 +190,8 @@ namespace UE::ConcertSyncServer::Replication
 				return bIsAllowed;
 			}
 
+			// TODO UE-219829: If EConcertReplicationPutStateFlags::SetStateForDisconnectedClients is set, remove mute changes for unknown objects
+			// and put them into the PutState activity, so the mute state can be applied in a RestoreContent request.
 			FPredictedStateObjectHierarchy FutureHierarchy;
 			FutureHierarchy.AddClients(Request.NewStreams);
 			FutureHierarchy.AddClients(Clients, [&Request](const FGuid& ClientId){ return !Request.NewStreams.Contains(ClientId); });
diff --git a/Engine/Plugins/Developer/Concert/ConcertSync/ConcertSyncServer/Source/ConcertSyncServer/Private/Replication/Util/GroundTruthOverride.cpp b/Engine/Plugins/Developer/Concert/ConcertSync/ConcertSyncServer/Source/ConcertSyncServer/Private/Replication/Util/GroundTruthOverride.cpp
--- a/Engine/Plugins/Developer/Concert/ConcertSync/ConcertSyncServer/Source/ConcertSyncServer/Private/Replication/Util/GroundTruthOverride.cpp
+++ b/Engine/Plugins/Developer/Concert/ConcertSync/ConcertSyncServer/Source/ConcertSyncServer/Private/Replication/Util/GroundTruthOverride.cpp
@@ -30,9 +30,34 @@ namespace UE::ConcertSyncServer::Replication
 		}
 	}
 
-	void FGroundTruthOverride::ForEachSendingClient(TFunctionRef<EBreakBehavior(const FGuid& ClientEndpointId)> Callback) const
+	void FGroundTruthOverride::ForEachClient(TFunctionRef<EBreakBehavior(const FGuid& ClientEndpointId)> Callback) const
 	{
-		NoOverrideStreamFallback.ForEachReplicationClient(Callback);
+		TSet<FGuid> AlreadyListed;
+
+		// First list every client that is actually registered on the server...
+		EBreakBehavior BreakBehavior = EBreakBehavior::Continue;
+		NoOverrideStreamFallback.ForEachReplicationClient([&Callback, &AlreadyListed, &BreakBehavior](const FGuid& EndpointId)
+		{
+			BreakBehavior = Callback(EndpointId);
+			AlreadyListed.Add(EndpointId);
+			return BreakBehavior;
+		});
+
+		if (BreakBehavior == EBreakBehavior::Break)
+		{
+			return;
+		}
+
+		// ... and then any further clients that are injected by the overrides (e.g. clients that have since disconnected).
+		for (const TPair<FGuid, FConcertReplicationStreamArray>& Overrides : StreamOverrides)
+		{
+			const FGuid& OverrideStream = Overrides.Key;
+			if (!AlreadyListed.Contains(OverrideStream)
+				&& Callback(OverrideStream) == EBreakBehavior::Break)
+			{
+				break;
+			}
+		}
 	}
 
 	bool FGroundTruthOverride::HasAuthority(const FGuid& ClientId, const FGuid& StreamId, const FSoftObjectPath& ObjectPath) const
diff --git a/Engine/Plugins/Developer/Concert/ConcertSync/ConcertSyncServer/Source/ConcertSyncServer/Private/Replication/Util/GroundTruthOverride.h b/Engine/Plugins/Developer/Concert/ConcertSync/ConcertSyncServer/Source/ConcertSyncServer/Private/Replication/Util/GroundTruthOverride.h
--- a/Engine/Plugins/Developer/Concert/ConcertSync/ConcertSyncServer/Source/ConcertSyncServer/Private/Replication/Util/GroundTruthOverride.h
+++ b/Engine/Plugins/Developer/Concert/ConcertSync/ConcertSyncServer/Source/ConcertSyncServer/Private/Replication/Util/GroundTruthOverride.h
@@ -22,6 +22,9 @@ namespace UE::ConcertSyncServer::Replication
 	/**
 	 * Pretends that the ground truth is the client overrides it was given.
 	 * If a setting is not overriden, then it defaults back to the server state.
+	 * 
+	 * StreamOverrides can contain clients that are not actually connected ("injection").
+	 * This is useful if you want to validate i.e. that no authority conflicts happen if the injected clients were present. 
 	 */
 	class FGroundTruthOverride
 		: public ConcertSyncCore::Replication::AuthorityConflictUtils::IReplicationGroundTruth
@@ -43,7 +46,7 @@ namespace UE::ConcertSyncServer::Replication
 
 		//~ Begin IReplicationGroundTruth Interface
 		virtual void ForEachStream(const FGuid& ClientEndpointId, TFunctionRef<EBreakBehavior(const FGuid& StreamId, const FConcertObjectReplicationMap& ReplicationMap)> Callback) const override;
-		virtual void ForEachSendingClient(TFunctionRef<EBreakBehavior(const FGuid& ClientEndpointId)> Callback) const override;
+		virtual void ForEachClient(TFunctionRef<EBreakBehavior(const FGuid& ClientEndpointId)> Callback) const override;
 		virtual bool HasAuthority(const FGuid& ClientId, const FGuid& StreamId, const FSoftObjectPath& ObjectPath) const override;
 		//~ Begin IReplicationGroundTruth Interface
 
diff --git a/Engine/Plugins/Developer/Concert/ConcertSync/ConcertSyncTest/Source/ConcertSyncTest/Private/Replication/ReplicationAuthorityTests.cpp b/Engine/Plugins/Developer/Concert/ConcertSync/ConcertSyncTest/Source/ConcertSyncTest/Private/Replication/ReplicationAuthorityTests.cpp
--- a/Engine/Plugins/Developer/Concert/ConcertSync/ConcertSyncTest/Source/ConcertSyncTest/Private/Replication/ReplicationAuthorityTests.cpp
+++ b/Engine/Plugins/Developer/Concert/ConcertSync/ConcertSyncTest/Source/ConcertSyncTest/Private/Replication/ReplicationAuthorityTests.cpp
@@ -400,7 +400,7 @@ namespace UE::ConcertSyncTests::Replication::Authority
 				}
 			}
 			
-			virtual void ForEachSendingClient(TFunctionRef<EBreakBehavior(const FGuid& ClientEndpointId)> Callback) const override
+			virtual void ForEachClient(TFunctionRef<EBreakBehavior(const FGuid& ClientEndpointId)> Callback) const override
 			{
 				if (bRequestingClientHasAuthority && Callback(RequestingClientId) == EBreakBehavior::Break)
 				{
