Refactor ValidateStreamChangeRequest to no longer accept FReplicationClient. Instead, pass in registered streams and endpoint ID.
This is in preparation for a future change.

diff --git a/Engine/Plugins/Developer/Concert/ConcertSync/ConcertSyncCore/Source/ConcertSyncCore/Private/Replication/ChangeStreamSharedUtils.cpp b/Engine/Plugins/Developer/Concert/ConcertSync/ConcertSyncCore/Source/ConcertSyncCore/Private/Replication/ChangeStreamSharedUtils.cpp
--- a/Engine/Plugins/Developer/Concert/ConcertSync/ConcertSyncCore/Source/ConcertSyncCore/Private/Replication/ChangeStreamSharedUtils.cpp
+++ b/Engine/Plugins/Developer/Concert/ConcertSync/ConcertSyncCore/Source/ConcertSyncCore/Private/Replication/ChangeStreamSharedUtils.cpp
@@ -215,7 +215,7 @@ namespace UE::ConcertSyncCore::Replication::ChangeStreamUtils
 
 	bool ValidateFrequencyChanges(
 		const FConcertReplication_ChangeStream_Request& Request,
-		const TArray<FConcertReplicationStream>& Streams,
+		const TConstArrayView<FConcertReplicationStream>& Streams,
 		FConcertReplication_ChangeStream_FrequencyResponse* OptionalErrors
 		)
 	{
diff --git a/Engine/Plugins/Developer/Concert/ConcertSync/ConcertSyncCore/Source/ConcertSyncCore/Public/Replication/ChangeStreamSharedUtils.h b/Engine/Plugins/Developer/Concert/ConcertSync/ConcertSyncCore/Source/ConcertSyncCore/Public/Replication/ChangeStreamSharedUtils.h
--- a/Engine/Plugins/Developer/Concert/ConcertSync/ConcertSyncCore/Source/ConcertSyncCore/Public/Replication/ChangeStreamSharedUtils.h
+++ b/Engine/Plugins/Developer/Concert/ConcertSync/ConcertSyncCore/Source/ConcertSyncCore/Public/Replication/ChangeStreamSharedUtils.h
@@ -80,7 +80,7 @@ namespace UE::ConcertSyncCore::Replication::ChangeStreamUtils
 	 */
 	CONCERTSYNCCORE_API bool ValidateFrequencyChanges(
 		const FConcertReplication_ChangeStream_Request& Request,
-		const TArray<FConcertReplicationStream>& Streams,
+		const TConstArrayView<FConcertReplicationStream>& Streams,
 		OUT FConcertReplication_ChangeStream_FrequencyResponse* OptionalErrors = nullptr
 		);
 
diff --git a/Engine/Plugins/Developer/Concert/ConcertSync/ConcertSyncServer/Source/ConcertSyncServer/Private/Replication/ConcertServerReplicationManager_ChangeStream.cpp b/Engine/Plugins/Developer/Concert/ConcertSync/ConcertSyncServer/Source/ConcertSyncServer/Private/Replication/ConcertServerReplicationManager_ChangeStream.cpp
--- a/Engine/Plugins/Developer/Concert/ConcertSync/ConcertSyncServer/Source/ConcertSyncServer/Private/Replication/ConcertServerReplicationManager_ChangeStream.cpp
+++ b/Engine/Plugins/Developer/Concert/ConcertSync/ConcertSyncServer/Source/ConcertSyncServer/Private/Replication/ConcertServerReplicationManager_ChangeStream.cpp
@@ -29,7 +29,7 @@ namespace UE::ConcertSyncServer::Replication
 		
 		const FGuid SendingClientId = ConcertSessionContext.SourceEndpointId;
 		const TUniquePtr<FConcertReplicationClient>* SendingClient = Clients.Find(SendingClientId);
-		if (SendingClient && ValidateStreamChangeRequest(*SendingClient->Get(), AuthorityManager, Request, Response))
+		if (SendingClient && ValidateStreamChangeRequest(SendingClient->Get()->GetClientEndpointId(), SendingClient->Get()->GetStreamDescriptions(), AuthorityManager, Request, Response))
 		{
 			ApplyChangeStreamRequest(Request, *SendingClient->Get());
 		}
diff --git a/Engine/Plugins/Developer/Concert/ConcertSync/ConcertSyncServer/Source/ConcertSyncServer/Private/Replication/ConcertServerReplicationManager_PutState.cpp b/Engine/Plugins/Developer/Concert/ConcertSync/ConcertSyncServer/Source/ConcertSyncServer/Private/Replication/ConcertServerReplicationManager_PutState.cpp
--- a/Engine/Plugins/Developer/Concert/ConcertSync/ConcertSyncServer/Source/ConcertSyncServer/Private/Replication/ConcertServerReplicationManager_PutState.cpp
+++ b/Engine/Plugins/Developer/Concert/ConcertSync/ConcertSyncServer/Source/ConcertSyncServer/Private/Replication/ConcertServerReplicationManager_PutState.cpp
@@ -129,7 +129,7 @@ namespace UE::ConcertSyncServer::Replication
 			const bool bSuccess = Algo::AllOf(StreamRequests, [&Clients, &AuthorityManager](const TPair<FGuid, FConcertReplication_ChangeStream_Request>& Pair)
 			{
 				const TUniquePtr<FConcertReplicationClient>* Client = Clients.Find(Pair.Key);
-				return !Client || ValidateStreamChangeRequest(*Client->Get(), AuthorityManager, Pair.Value);
+				return !Client || ValidateStreamChangeRequest(Client->Get()->GetClientEndpointId(), Client->Get()->GetStreamDescriptions(), AuthorityManager, Pair.Value);
 			});
 
 			Response.ResponseCode = bSuccess ? Response.ResponseCode : EConcertReplicationPutStateResponseCode::StreamError;
diff --git a/Engine/Plugins/Developer/Concert/ConcertSync/ConcertSyncServer/Source/ConcertSyncServer/Private/Replication/Util/StreamChangeValidation.cpp b/Engine/Plugins/Developer/Concert/ConcertSync/ConcertSyncServer/Source/ConcertSyncServer/Private/Replication/Util/StreamChangeValidation.cpp
--- a/Engine/Plugins/Developer/Concert/ConcertSync/ConcertSyncServer/Source/ConcertSyncServer/Private/Replication/Util/StreamChangeValidation.cpp
+++ b/Engine/Plugins/Developer/Concert/ConcertSync/ConcertSyncServer/Source/ConcertSyncServer/Private/Replication/Util/StreamChangeValidation.cpp
@@ -13,16 +13,16 @@ namespace UE::ConcertSyncServer::Replication
 {
 	namespace Private
 	{
-		static const FConcertReplicationStream* FindExistingStream(const FConcertReplicationClient& Client, const FGuid& StreamId)
+		static const FConcertReplicationStream* FindExistingStream(const TConstArrayView<FConcertReplicationStream>& Streams, const FGuid& StreamId)
 		{
-			return Client.GetStreamDescriptions().FindByPredicate([&StreamId](const FConcertReplicationStream& Description)
+			return Streams.FindByPredicate([&StreamId](const FConcertReplicationStream& Description)
 			{
 				return Description.BaseDescription.Identifier == StreamId;
 			});
 		}
 		
 		/** Validates that ObjectsToPut writes only to pre-existing streams. */
-		static void ValidatePutObjectsRequestSemantics(const FConcertReplication_ChangeStream_Request& Request, const FConcertReplicationClient& Client, FConcertReplication_ChangeStream_Response& OutResponse)
+		static void ValidatePutObjectsRequestSemantics(const TConstArrayView<FConcertReplicationStream>& Streams, const FConcertReplication_ChangeStream_Request& Request, FConcertReplication_ChangeStream_Response& OutResponse)
 		{
 			const auto PutObjectHasEnoughData = [](const FConcertReplicationStream* ExistingStream, const FSoftObjectPath& ChangedObjectPath, const FConcertReplication_ChangeStream_PutObject& PutObject)
 			{
@@ -40,7 +40,7 @@ namespace UE::ConcertSyncServer::Replication
 				const FConcertObjectInStreamID ChangedObject = Change.Key;
 				
 				const FGuid& StreamToModify = ChangedObject.StreamId;
-				const FConcertReplicationStream* ExistingStream = FindExistingStream(Client, StreamToModify);
+				const FConcertReplicationStream* ExistingStream = FindExistingStream(Streams, StreamToModify);
 				const bool bStreamExists = ExistingStream != nullptr;
 				if (bStreamExists)
 				{
@@ -63,14 +63,14 @@ namespace UE::ConcertSyncServer::Replication
 		}
 
 		/** Checks that StreamsToAdd do not conflict with pre-existing ones and that all IDs in the request are also unique. */
-		static void ValidateAddedStreamsAreValid(const FConcertReplication_ChangeStream_Request& Request, const FConcertReplicationClient& Client, FConcertReplication_ChangeStream_Response& OutResponse)
+		static void ValidateAddedStreamsAreValid(const TConstArrayView<FConcertReplicationStream>& Streams, const FConcertReplication_ChangeStream_Request& Request, FConcertReplication_ChangeStream_Response& OutResponse)
 		{
 			TSet<FGuid> DuplicateEntryDetection;
 			for (const FConcertReplicationStream& NewStream : Request.StreamsToAdd)
 			{
 				// StreamsToAdd is invalid if there is already a stream with the same ID registered ...
 				const FGuid& NewStreamId = NewStream.BaseDescription.Identifier;
-				const bool bIdAlreadyExists = FindExistingStream(Client, NewStreamId) != nullptr;
+				const bool bIdAlreadyExists = FindExistingStream(Streams, NewStreamId) != nullptr;
 				const bool bIsStreamRemoved = Request.StreamsToRemove.Contains(NewStreamId);
 
 				// ... or StreamsToAdd contains the same ID multiple times
@@ -88,9 +88,8 @@ namespace UE::ConcertSyncServer::Replication
 		}
 
 		/** If the requesting client has authority over a changed object, checks that no other client has authority over the properties being added. */
-		static void LookForAuthorityConflicts(const FConcertReplication_ChangeStream_Request& Request, const FConcertReplicationClient& Client, const FAuthorityManager& AuthorityManager, FConcertReplication_ChangeStream_Response& OutResponse)
+		static void LookForAuthorityConflicts(const FGuid& ClientEndpointId, const FConcertReplication_ChangeStream_Request& Request, const FAuthorityManager& AuthorityManager, FConcertReplication_ChangeStream_Response& OutResponse)
 		{
-			const FGuid ClientEndpointId = Client.GetClientEndpointId();
 			for (const TPair<FConcertObjectInStreamID, FConcertReplication_ChangeStream_PutObject>& PutObjectPair : Request.ObjectsToPut)
 			{
 				// No conflict possible if requesting client does not have authority over the changed object
@@ -123,7 +122,8 @@ namespace UE::ConcertSyncServer::Replication
 	}
 	
 	bool ValidateStreamChangeRequest(
-		const FConcertReplicationClient& Client,
+		const FGuid& ClientEndpointId,
+		const TConstArrayView<FConcertReplicationStream>& Streams,
 		const FAuthorityManager& AuthorityManager,
 		const FConcertReplication_ChangeStream_Request& Request,
 		FConcertReplication_ChangeStream_Response& OutResponse
@@ -131,17 +131,22 @@ namespace UE::ConcertSyncServer::Replication
 	{
 		OutResponse.ErrorCode = EReplicationResponseErrorCode::Handled;
 
-		Private::ValidatePutObjectsRequestSemantics(Request, Client, OutResponse);
-		Private::ValidateAddedStreamsAreValid(Request, Client, OutResponse);
-		Private::LookForAuthorityConflicts(Request, Client, AuthorityManager, OutResponse);
-		ConcertSyncCore::Replication::ChangeStreamUtils::ValidateFrequencyChanges(Request, Client.GetStreamDescriptions(), &OutResponse.FrequencyErrors);
+		Private::ValidatePutObjectsRequestSemantics(Streams, Request, OutResponse);
+		Private::ValidateAddedStreamsAreValid(Streams, Request, OutResponse);
+		Private::LookForAuthorityConflicts(ClientEndpointId, Request, AuthorityManager, OutResponse);
+		ConcertSyncCore::Replication::ChangeStreamUtils::ValidateFrequencyChanges(Request, Streams, &OutResponse.FrequencyErrors);
 			
 		return OutResponse.IsSuccess();
 	}
 
-	bool ValidateStreamChangeRequest(const FConcertReplicationClient& Client, const FAuthorityManager& AuthorityManager, const FConcertReplication_ChangeStream_Request& Request)
+	bool ValidateStreamChangeRequest(
+		const FGuid& ClientEndpointId,
+		const TConstArrayView<FConcertReplicationStream>& Streams,
+		const FAuthorityManager& AuthorityManager,
+		const FConcertReplication_ChangeStream_Request& Request
+		)
 	{
 		FConcertReplication_ChangeStream_Response Dummy;
-		return ValidateStreamChangeRequest(Client, AuthorityManager, Request, Dummy);
+		return ValidateStreamChangeRequest(ClientEndpointId, Streams, AuthorityManager, Request, Dummy);
 	}
 }
diff --git a/Engine/Plugins/Developer/Concert/ConcertSync/ConcertSyncServer/Source/ConcertSyncServer/Private/Replication/Util/StreamChangeValidation.h b/Engine/Plugins/Developer/Concert/ConcertSync/ConcertSyncServer/Source/ConcertSyncServer/Private/Replication/Util/StreamChangeValidation.h
--- a/Engine/Plugins/Developer/Concert/ConcertSync/ConcertSyncServer/Source/ConcertSyncServer/Private/Replication/Util/StreamChangeValidation.h
+++ b/Engine/Plugins/Developer/Concert/ConcertSync/ConcertSyncServer/Source/ConcertSyncServer/Private/Replication/Util/StreamChangeValidation.h
@@ -2,6 +2,10 @@
 
 #pragma once
 
+#include "Containers/ArrayView.h"
+#include "Containers/ContainersFwd.h"
+
+struct FConcertReplicationStream;
 struct FConcertReplication_ChangeStream_Response;
 struct FConcertReplication_ChangeStream_Request;
 
@@ -15,7 +19,8 @@ namespace UE::ConcertSyncServer::Replication
 	 * @return Whether the request is valid.
 	 */
 	bool ValidateStreamChangeRequest(
-		const FConcertReplicationClient& Client,
+		const FGuid& ClientEndpointId,
+		const TConstArrayView<FConcertReplicationStream>& Streams,
 		const FAuthorityManager& AuthorityManager,
 		const FConcertReplication_ChangeStream_Request& Request,
 		FConcertReplication_ChangeStream_Response& OutResponse
@@ -23,7 +28,8 @@ namespace UE::ConcertSyncServer::Replication
 	
 	/** @return Whether the request is valid. */
 	bool ValidateStreamChangeRequest(
-		const FConcertReplicationClient& Client,
+		const FGuid& ClientEndpointId,
+		const TConstArrayView<FConcertReplicationStream>& Streams,
 		const FAuthorityManager& AuthorityManager,
 		const FConcertReplication_ChangeStream_Request& Request
 		);
