Make functions in IReplicationWorkspace more generic where possible to increase flexibility and make it easier to add new activities.

diff --git a/Engine/Plugins/Developer/Concert/ConcertSync/ConcertSyncCore/Source/ConcertSyncCore/Public/Replication/Messages/ReplicationActivity.h b/Engine/Plugins/Developer/Concert/ConcertSync/ConcertSyncCore/Source/ConcertSyncCore/Public/Replication/Messages/ReplicationActivity.h
--- a/Engine/Plugins/Developer/Concert/ConcertSync/ConcertSyncCore/Source/ConcertSyncCore/Public/Replication/Messages/ReplicationActivity.h
+++ b/Engine/Plugins/Developer/Concert/ConcertSync/ConcertSyncCore/Source/ConcertSyncCore/Public/Replication/Messages/ReplicationActivity.h
@@ -120,7 +120,8 @@ struct FConcertSyncReplicationEvent
 	FConcertSessionSerializedPayload Payload{ EConcertPayloadSerializationMethod::Cbor };
 
 	FConcertSyncReplicationEvent() = default;
-	FConcertSyncReplicationEvent(const FConcertSyncReplicationPayload_LeaveReplication& Data)
+	template<typename TPayload>
+	explicit FConcertSyncReplicationEvent(const TPayload& Data)
 	{
 		SetPayload(Data);
 	}
@@ -162,7 +163,8 @@ struct FConcertSyncReplicationActivity : public FConcertSyncActivity
 		EventType = EConcertSyncActivityEventType::Replication;
 	}
 
-	FConcertSyncReplicationActivity(const FConcertSyncReplicationPayload_LeaveReplication& PayloadData)
+	template<typename TPayload>
+	explicit FConcertSyncReplicationActivity(const TPayload& PayloadData)
 		: EventData(PayloadData)
 	{
 		EventType = EConcertSyncActivityEventType::Replication;
diff --git a/Engine/Plugins/Developer/Concert/ConcertSync/ConcertSyncServer/Source/ConcertSyncServer/Private/Replication/ConcertServerReplicationManager_RestoreContent.cpp b/Engine/Plugins/Developer/Concert/ConcertSync/ConcertSyncServer/Source/ConcertSyncServer/Private/Replication/ConcertServerReplicationManager_RestoreContent.cpp
--- a/Engine/Plugins/Developer/Concert/ConcertSync/ConcertSyncServer/Source/ConcertSyncServer/Private/Replication/ConcertServerReplicationManager_RestoreContent.cpp
+++ b/Engine/Plugins/Developer/Concert/ConcertSync/ConcertSyncServer/Source/ConcertSyncServer/Private/Replication/ConcertServerReplicationManager_RestoreContent.cpp
@@ -244,7 +244,7 @@ namespace UE::ConcertSyncServer::Replication
 		FConcertSyncReplicationPayload_LeaveReplication OldClientState;
 		const bool bOverrideActivity = Request.ActivityId.IsSet();
 		const bool bFound = bOverrideActivity
-			? ServerWorkspace.GetLeaveReplicationActivityById(*Request.ActivityId, OldClientState)
+			? ServerWorkspace.GetLeaveReplicationEventById(*Request.ActivityId, OldClientState)
 			: ServerWorkspace.GetLastLeaveReplicationActivityByClient(ClientInfo, OldClientState);
 		if (!bFound)
 		{
diff --git a/Engine/Plugins/Developer/Concert/ConcertSync/ConcertSyncServer/Source/ConcertSyncServer/Private/Replication/ReplicationWorkspace.cpp b/Engine/Plugins/Developer/Concert/ConcertSync/ConcertSyncServer/Source/ConcertSyncServer/Private/Replication/ReplicationWorkspace.cpp
--- a/Engine/Plugins/Developer/Concert/ConcertSync/ConcertSyncServer/Source/ConcertSyncServer/Private/Replication/ReplicationWorkspace.cpp
+++ b/Engine/Plugins/Developer/Concert/ConcertSync/ConcertSyncServer/Source/ConcertSyncServer/Private/Replication/ReplicationWorkspace.cpp
@@ -19,17 +19,24 @@ namespace UE::ConcertSyncServer
 		check(FindSessionClientDelegate.IsBound() && ShouldIgnoreClientActivityOnRestoreDelegate.IsBound());
 	}
 
-	TOptional<int64> FReplicationWorkspace::ProduceClientLeaveReplicationActivity(const FGuid& EndpointId, const FConcertSyncReplicationPayload_LeaveReplication& EventData)
+	TOptional<int64> FReplicationWorkspace::ProduceReplicationActivity(const FGuid& EndpointId, const FConcertSyncReplicationEvent& EventData)
 	{
-		return ProduceActivity(EndpointId, EventData);
-	}
+		FConcertSyncReplicationActivity Activity;
+		Activity.EndpointId = EndpointId;
+		Activity.EventData = EventData;
+		Activity.EventSummary.SetTypedPayload(FConcertSyncReplicationActivitySummary::CreateSummaryForEvent(Activity.EventData));
+		Activity.bIgnored = ShouldIgnoreClientActivityOnRestoreDelegate.Execute(EndpointId);
+		
+		int64 ActivityId = 0;
+		int64 EventId = 0;
+		const bool bSuccess = Database.AddReplicationActivity(Activity, ActivityId, EventId);
+		// This allows FConcertServerWorkspace to send the activity to other clients
+		OnAddReplicationActivityDelegate.Broadcast(ActivityId, bSuccess);
 
-	TOptional<int64> FReplicationWorkspace::ProduceClientMuteReplicationActivity(const FGuid& EndpointId, const FConcertSyncReplicationPayload_Mute& EventData)
-	{
-		return ProduceActivity(EndpointId, EventData);
+		return bSuccess ? ActivityId : TOptional<int64>{};
 	}
 
-	bool FReplicationWorkspace::GetLastLeaveReplicationActivityByClient(const FConcertSessionClientInfo& InClientInfo, FConcertSyncReplicationPayload_LeaveReplication& OutLeaveReplication) const
+	bool FReplicationWorkspace::GetLastReplicationActivityByClient(const FConcertSessionClientInfo& InClientInfo, FConcertSyncReplicationActivity& OutActivity) const
 	{
 		// Goal: Iterate all endpoints with the same display name. Then, get the latest transaction ID from all of them.
 		// This entire approach is suboptimal though: There should be a dedicated SQL prepared statement  to which we give the client and device name.
@@ -61,44 +68,20 @@ namespace UE::ConcertSyncServer
 			NewestEventId = FMath::Max(NewestEventId, EventId);
 		}
 
-		FConcertSyncReplicationEvent Event;
-		const bool bGotData = NewestEventId != INDEX_NONE && Database.GetReplicationEvent(NewestEventId, Event);
-		return bGotData
-			&& Event.ActivityType == EConcertSyncReplicationActivityType::LeaveReplication
-			&& Event.GetPayload(OutLeaveReplication);
+		const bool bGotData = NewestEventId != INDEX_NONE && Database.GetReplicationActivity(NewestEventId, OutActivity);
+		return bGotData;
 	}
 
-	bool FReplicationWorkspace::GetLeaveReplicationActivityById(const int64 ActivityId, FConcertSyncReplicationPayload_LeaveReplication& OutLeaveReplication) const
+	bool FReplicationWorkspace::GetReplicationEventById(const int64 ActivityId, FConcertSyncReplicationEvent& OutEvent) const
 	{
-		FConcertSyncReplicationEvent Event;
-		const bool bGotEventData = Database.GetReplicationEvent(ActivityId, Event);
-		return bGotEventData && Event.ActivityType == EConcertSyncReplicationActivityType::LeaveReplication && Event.GetPayload(OutLeaveReplication);
+		return Database.GetReplicationEvent(ActivityId, OutEvent);
 	}
 
-	void FReplicationWorkspace::EnumerateMuteActivities(TFunctionRef<EBreakBehavior(const FConcertSyncReplicationActivity& Activity)> Callback) const
+	void FReplicationWorkspace::EnumerateReplicationActivities(TFunctionRef<EBreakBehavior(const FConcertSyncReplicationActivity& Activity)> Callback) const
 	{
 		Database.EnumerateReplicationActivities([&Callback](FConcertSyncReplicationActivity&& Activity)
 		{
-			return Activity.EventData.ActivityType != EConcertSyncReplicationActivityType::Mute
-				|| Callback(Activity) == EBreakBehavior::Continue;
+			return Callback(Activity) == EBreakBehavior::Continue;
 		});
 	}
-
-	template <typename TPayload>
-	TOptional<int64> FReplicationWorkspace::ProduceActivity(const FGuid& EndpointId, const TPayload& EventData)
-	{
-		FConcertSyncReplicationActivity Activity;
-		Activity.EndpointId = EndpointId;
-		Activity.EventData.SetPayload(EventData);
-		Activity.EventSummary.SetTypedPayload(FConcertSyncReplicationActivitySummary::CreateSummaryForEvent(Activity.EventData));
-		Activity.bIgnored = ShouldIgnoreClientActivityOnRestoreDelegate.Execute(EndpointId);
-		
-		int64 ActivityId = 0;
-		int64 EventId = 0;
-		const bool bSuccess = Database.AddReplicationActivity(Activity, ActivityId, EventId);
-		// This allows FConcertServerWorkspace to send the activity to other clients
-		OnAddReplicationActivityDelegate.Broadcast(ActivityId, bSuccess);
-
-		return bSuccess ? ActivityId : TOptional<int64>{};
-	}
 }
diff --git a/Engine/Plugins/Developer/Concert/ConcertSync/ConcertSyncServer/Source/ConcertSyncServer/Private/Replication/ReplicationWorkspace.h b/Engine/Plugins/Developer/Concert/ConcertSync/ConcertSyncServer/Source/ConcertSyncServer/Private/Replication/ReplicationWorkspace.h
--- a/Engine/Plugins/Developer/Concert/ConcertSync/ConcertSyncServer/Source/ConcertSyncServer/Private/Replication/ReplicationWorkspace.h
+++ b/Engine/Plugins/Developer/Concert/ConcertSync/ConcertSyncServer/Source/ConcertSyncServer/Private/Replication/ReplicationWorkspace.h
@@ -37,11 +37,10 @@ namespace UE::ConcertSyncServer
 		FOnAddReplicationActivity& OnAddReplicationActivity() { return OnAddReplicationActivityDelegate; }
 
 		//~ Begin IReplicationWorkspace Interface
-		virtual TOptional<int64> ProduceClientLeaveReplicationActivity(const FGuid& EndpointId, const FConcertSyncReplicationPayload_LeaveReplication& EventData) override;
-		virtual TOptional<int64> ProduceClientMuteReplicationActivity(const FGuid& EndpointId, const FConcertSyncReplicationPayload_Mute& EventData) override;
-		virtual bool GetLastLeaveReplicationActivityByClient(const FConcertSessionClientInfo& InClientInfo, FConcertSyncReplicationPayload_LeaveReplication& OutLeaveReplication) const override;
-		virtual bool GetLeaveReplicationActivityById(const int64 ActivityId, FConcertSyncReplicationPayload_LeaveReplication& OutLeaveReplication) const override;
-		virtual void EnumerateMuteActivities(TFunctionRef<EBreakBehavior(const FConcertSyncReplicationActivity& Activity)> Callback) const override;
+		virtual TOptional<int64> ProduceReplicationActivity(const FGuid& EndpointId, const FConcertSyncReplicationEvent& EventData) override;
+		virtual bool GetLastReplicationActivityByClient(const FConcertSessionClientInfo& InClientInfo, FConcertSyncReplicationActivity& OutActivity) const override;
+		virtual bool GetReplicationEventById(const int64 ActivityId, FConcertSyncReplicationEvent& OutEvent) const override;
+		virtual void EnumerateReplicationActivities(TFunctionRef<EBreakBehavior(const FConcertSyncReplicationActivity& Activity)> Callback) const override;
 		//~ End IReplicationWorkspace Interface
 
 	private:
@@ -54,9 +53,6 @@ namespace UE::ConcertSyncServer
 		const FShouldIgnoreClientActivityOnRestore ShouldIgnoreClientActivityOnRestoreDelegate;
 		
 		FOnAddReplicationActivity OnAddReplicationActivityDelegate;
-
-		template<typename TPayload>
-		TOptional<int64> ProduceActivity(const FGuid& EndpointId, const TPayload& EventData);
 	};
 }
 
diff --git a/Engine/Plugins/Developer/Concert/ConcertSync/ConcertSyncServer/Source/ConcertSyncServer/Public/Replication/IReplicationWorkspace.h b/Engine/Plugins/Developer/Concert/ConcertSync/ConcertSyncServer/Source/ConcertSyncServer/Public/Replication/IReplicationWorkspace.h
--- a/Engine/Plugins/Developer/Concert/ConcertSync/ConcertSyncServer/Source/ConcertSyncServer/Public/Replication/IReplicationWorkspace.h
+++ b/Engine/Plugins/Developer/Concert/ConcertSync/ConcertSyncServer/Source/ConcertSyncServer/Public/Replication/IReplicationWorkspace.h
@@ -4,8 +4,10 @@
 
 #include "HAL/Platform.h"
 #include "Misc/Optional.h"
+#include "Replication/Messages/ReplicationActivity.h"
 #include "Templates/FunctionFwd.h"
 
+struct FConcertSyncReplicationEvent;
 enum class EBreakBehavior : uint8;
 
 struct FConcertSyncReplicationActivity;
@@ -25,34 +27,76 @@ namespace UE::ConcertSyncServer::Replication
 	public:
 		
 		/**
-		 * Creates a replication activity for the client leaving replication
-		 * @return The identifier of the produced activity. Unset if activity insertion failed.
-		 */
-		virtual TOptional<int64> ProduceClientLeaveReplicationActivity(const FGuid& EndpointId, const FConcertSyncReplicationPayload_LeaveReplication& EventData) = 0;
+		 * Creates a replication activity for the provided client.
+		 * @param EndpointId The client that produced the activity
+		 * @param EventData Data associated with the activity - must have ActivityType other than EConcertSyncReplicationActivityType::None.
+         * @return The identifier of the produced activity. Unset if activity insertion failed.
+         */
+		virtual TOptional<int64> ProduceReplicationActivity(const FGuid& EndpointId, const FConcertSyncReplicationEvent& EventData) = 0;
+		/** Creates a replication activity for the client leaving replication. */
+		TOptional<int64> ProduceClientLeaveReplicationActivity(const FGuid& EndpointId, const FConcertSyncReplicationPayload_LeaveReplication& EventData);
+		/** Creates a replication activity for the client (un)muting objects in the session. */
+		TOptional<int64> ProduceClientMuteReplicationActivity(const FGuid& EndpointId, const FConcertSyncReplicationPayload_Mute& EventData);
 
 		/**
-		 * Creates a replication activity for the client (un)muting objects in the session.
-		 * @return The identifier of the produced activity. Unset if activity insertion failed.
-		 */
-		virtual TOptional<int64> ProduceClientMuteReplicationActivity(const FGuid& EndpointId, const FConcertSyncReplicationPayload_Mute& EventData) = 0;
-
-		/**
-		 * Gets the last replication leave activity associated for a given client.
+		 * Gets the last replication activity associated with the given client info.
 		 * 
 		 * As endpoint IDs change every time a client join a session, the look up is done by client display name.
 		 * If multiple machines joined with the same display name, the tie is broken by also using the device name.
 		 * 
 		 * @param InClientInfo Info about the client for which to get the activity.
-		 * @param OutLeaveReplication The activity, if present
-		 * @return Whether OutLeaveReplication contains a valid result.
+		 * @param OutActivity The activity, if present
+		 * @return Whether OutActivity contains a valid result.
 		 */
-		virtual bool GetLastLeaveReplicationActivityByClient(const FConcertSessionClientInfo& InClientInfo, FConcertSyncReplicationPayload_LeaveReplication& OutLeaveReplication) const = 0;
+		virtual bool GetLastReplicationActivityByClient(const FConcertSessionClientInfo& InClientInfo, FConcertSyncReplicationActivity& OutActivity) const = 0;
+		/** Gets the last leave replication activity associated with the given client info. */
+		bool GetLastLeaveReplicationActivityByClient(const FConcertSessionClientInfo& InClientInfo, FConcertSyncReplicationPayload_LeaveReplication& OutLeaveReplication) const;
+		
 		/** Gets the replication leave activity with ActivityId. */
-		virtual bool GetLeaveReplicationActivityById(const int64 ActivityId, FConcertSyncReplicationPayload_LeaveReplication& OutLeaveReplication) const = 0;
+		virtual bool GetReplicationEventById(const int64 ActivityId, FConcertSyncReplicationEvent& OutEvent) const = 0;
+		/** Gets the replication leave activity with ActivityId. */
+		bool GetLeaveReplicationEventById(const int64 ActivityId, FConcertSyncReplicationPayload_LeaveReplication& OutLeaveReplication) const;
 
+		/** Enumerates all activities. */
+		virtual void EnumerateReplicationActivities(TFunctionRef<EBreakBehavior(const FConcertSyncReplicationActivity& Activity)> Callback) const = 0;
 		/** Enumerates all mute activities. */
-		virtual void EnumerateMuteActivities(TFunctionRef<EBreakBehavior(const FConcertSyncReplicationActivity& Activity)> Callback) const = 0;
+		void EnumerateMuteActivities(TFunctionRef<EBreakBehavior(const FConcertSyncReplicationActivity& Activity)> Callback) const;
 
 		virtual ~IReplicationWorkspace() = default;
 	};
+
+
+	inline TOptional<int64> IReplicationWorkspace::ProduceClientLeaveReplicationActivity(const FGuid& EndpointId, const FConcertSyncReplicationPayload_LeaveReplication& EventData)
+	{
+		return ProduceReplicationActivity(EndpointId, FConcertSyncReplicationEvent(EventData));
+	}
+
+	inline TOptional<int64> IReplicationWorkspace::ProduceClientMuteReplicationActivity(const FGuid& EndpointId, const FConcertSyncReplicationPayload_Mute& EventData)
+	{
+		return ProduceReplicationActivity(EndpointId, FConcertSyncReplicationEvent(EventData));
+	}
+
+	inline bool IReplicationWorkspace::GetLastLeaveReplicationActivityByClient(const FConcertSessionClientInfo& InClientInfo, FConcertSyncReplicationPayload_LeaveReplication& OutLeaveReplication) const
+	{
+		FConcertSyncReplicationActivity Activity;
+		return GetLastReplicationActivityByClient(InClientInfo, Activity)
+			&& ensureMsgf(Activity.EventData.ActivityType == EConcertSyncReplicationActivityType::LeaveReplication, TEXT("Caller expected ActivityId %lld to be a LeaveReplication event"), Activity.ActivityId)
+			&& Activity.EventData.GetPayload(OutLeaveReplication);
+	}
+
+	inline bool IReplicationWorkspace::GetLeaveReplicationEventById(const int64 ActivityId, FConcertSyncReplicationPayload_LeaveReplication& OutLeaveReplication) const
+	{
+		FConcertSyncReplicationEvent Event;
+		return GetReplicationEventById(ActivityId, Event)
+			&& ensureMsgf(Event.ActivityType == EConcertSyncReplicationActivityType::LeaveReplication, TEXT("Caller expected ActivityId %lld to be a LeaveReplication event"), ActivityId)
+			&& Event.GetPayload(OutLeaveReplication);
+	}
+
+	inline void IReplicationWorkspace::EnumerateMuteActivities(TFunctionRef<EBreakBehavior(const FConcertSyncReplicationActivity& Activity)> Callback) const
+	{
+		return EnumerateReplicationActivities([&Callback](const FConcertSyncReplicationActivity& Activity)
+		{
+			return Activity.EventData.ActivityType == EConcertSyncReplicationActivityType::Mute ? Callback(Activity) : EBreakBehavior::Continue;
+		});
+	}
 }
diff --git a/Engine/Plugins/Developer/Concert/ConcertSync/ConcertSyncTest/Source/ConcertSyncTest/Private/Replication/ReplicationWorkspace.spec.cpp b/Engine/Plugins/Developer/Concert/ConcertSync/ConcertSyncTest/Source/ConcertSyncTest/Private/Replication/ReplicationWorkspace.spec.cpp
--- a/Engine/Plugins/Developer/Concert/ConcertSync/ConcertSyncTest/Source/ConcertSyncTest/Private/Replication/ReplicationWorkspace.spec.cpp
+++ b/Engine/Plugins/Developer/Concert/ConcertSync/ConcertSyncTest/Source/ConcertSyncTest/Private/Replication/ReplicationWorkspace.spec.cpp
@@ -155,7 +155,7 @@ namespace UE::ConcertSyncTests::Replication::UI
 			}
 
 			FConcertSyncReplicationPayload_LeaveReplication ActualData;
-			const bool bGotData = ReplicationWorkspace->GetLeaveReplicationActivityById(*ActivityId, ActualData);
+			const bool bGotData = ReplicationWorkspace->GetLeaveReplicationEventById(*ActivityId, ActualData);
 			TestTrue(TEXT("Got Data"), bGotData);
 			TestEqual(TEXT("Payload"), ActualData, ExpectedData);
 		});
diff --git a/Engine/Plugins/Developer/Concert/ConcertSync/ConcertSyncTest/Source/ConcertSyncTest/Private/Replication/Restore/RestoreContentPreexisting.spec.cpp b/Engine/Plugins/Developer/Concert/ConcertSync/ConcertSyncTest/Source/ConcertSyncTest/Private/Replication/Restore/RestoreContentPreexisting.spec.cpp
--- a/Engine/Plugins/Developer/Concert/ConcertSync/ConcertSyncTest/Source/ConcertSyncTest/Private/Replication/Restore/RestoreContentPreexisting.spec.cpp
+++ b/Engine/Plugins/Developer/Concert/ConcertSync/ConcertSyncTest/Source/ConcertSyncTest/Private/Replication/Restore/RestoreContentPreexisting.spec.cpp
@@ -51,8 +51,8 @@ namespace UE::ConcertSyncTests::Replication::RestoreContent
 	
 		void InsertActivityData() const
 		{
-			WorkspaceMock->ReturnResult_GetLastLeaveReplicationActivityByClient = LeaveReplicationData;
-			WorkspaceMock->ReturnResult_GetLeaveReplicationActivityById = LeaveReplicationData;
+			WorkspaceMock->ReturnResult_GetLastReplicationActivityByClient = FConcertSyncReplicationActivity(LeaveReplicationData);
+			WorkspaceMock->ReturnResult_GetReplicationEventById = FConcertSyncReplicationEvent(LeaveReplicationData);
 		}
 		FConcertReplicatedObjectInfo& GetRestoredObjectData()
 		{
diff --git a/Engine/Plugins/Developer/Concert/ConcertSync/ConcertSyncTest/Source/ConcertSyncTest/Private/Replication/Restore/RestoreContentWithNoPreexisting.spec.cpp b/Engine/Plugins/Developer/Concert/ConcertSync/ConcertSyncTest/Source/ConcertSyncTest/Private/Replication/Restore/RestoreContentWithNoPreexisting.spec.cpp
--- a/Engine/Plugins/Developer/Concert/ConcertSync/ConcertSyncTest/Source/ConcertSyncTest/Private/Replication/Restore/RestoreContentWithNoPreexisting.spec.cpp
+++ b/Engine/Plugins/Developer/Concert/ConcertSync/ConcertSyncTest/Source/ConcertSyncTest/Private/Replication/Restore/RestoreContentWithNoPreexisting.spec.cpp
@@ -33,8 +33,8 @@ namespace UE::ConcertSyncTests::Replication::RestoreContent
 
 		void InsertActivityData() const
 		{
-			WorkspaceMock->ReturnResult_GetLastLeaveReplicationActivityByClient = LeaveReplicationData;
-			WorkspaceMock->ReturnResult_GetLeaveReplicationActivityById = LeaveReplicationData;
+			WorkspaceMock->ReturnResult_GetLastReplicationActivityByClient = FConcertSyncReplicationActivity(LeaveReplicationData);
+			WorkspaceMock->ReturnResult_GetReplicationEventById = FConcertSyncReplicationEvent(LeaveReplicationData);
 		}
 
 		TFuture<FConcertReplication_RestoreContent_Response> RestoreThenTestErrorCode(const FConcertReplication_RestoreContent_Request& Request, EConcertReplicationRestoreErrorCode ExpectedErrorCode)
diff --git a/Engine/Plugins/Developer/Concert/ConcertSync/ConcertSyncTest/Source/ConcertSyncTest/Private/Replication/Util/Mocks/ReplicationWorkspaceCallInterceptorMock.h b/Engine/Plugins/Developer/Concert/ConcertSync/ConcertSyncTest/Source/ConcertSyncTest/Private/Replication/Util/Mocks/ReplicationWorkspaceCallInterceptorMock.h
--- a/Engine/Plugins/Developer/Concert/ConcertSync/ConcertSyncTest/Source/ConcertSyncTest/Private/Replication/Util/Mocks/ReplicationWorkspaceCallInterceptorMock.h
+++ b/Engine/Plugins/Developer/Concert/ConcertSync/ConcertSyncTest/Source/ConcertSyncTest/Private/Replication/Util/Mocks/ReplicationWorkspaceCallInterceptorMock.h
@@ -21,7 +21,7 @@ namespace UE::ConcertSyncTests::Replication
 		/** Arguments of last ProduceClientMuteReplicationActivity call. */
 		TOptional<TTuple<FGuid, FConcertSyncReplicationPayload_Mute>> LastCall_ProduceClientMuteReplicationActivity;
 		/** Arguments of last GetLastLeaveReplicationActivityByClient call. */
-		mutable TOptional<TTuple<FConcertSessionClientInfo>> LastCall_GetLastLeaveReplicationActivityByClient;
+		mutable TOptional<TTuple<FConcertSessionClientInfo>> LastCall_GetLastReplicationActivityByClient;
 		/** Arguments of last GetLastLeaveReplicationActivityByClient call. */
 		mutable TOptional<TTuple<int64>> LastCall_GetLeaveReplicationActivityById;
 
@@ -29,51 +29,71 @@ namespace UE::ConcertSyncTests::Replication
 		TOptional<int64> ReturnResult_ProduceClientLeaveReplicationActivity = 0;
 		/** The result to return in ProduceClientLeaveReplicationActivity. */
 		TOptional<int64> ReturnResult_ProduceClientMuteReplicationActivity = 0;
-		/** The result to return in GetLastLeaveReplicationActivityByClient. */
-		TOptional<FConcertSyncReplicationPayload_LeaveReplication> ReturnResult_GetLastLeaveReplicationActivityByClient;
-		/** The result to return in GetLastLeaveReplicationActivityByClient. */
-		TOptional<FConcertSyncReplicationPayload_LeaveReplication> ReturnResult_GetLeaveReplicationActivityById;
+		/** The result to return in GetLastReplicationActivityByClient. */
+		TOptional<FConcertSyncReplicationActivity> ReturnResult_GetLastReplicationActivityByClient;
+		/** The result to return in GetReplicationEventById. */
+		TOptional<FConcertSyncReplicationEvent> ReturnResult_GetReplicationEventById;
 		/** The values to enumerate in EnumerateMuteActivities. */
-		TOptional<TArray<FConcertSyncReplicationActivity>> ReturnResult_EnumerateMuteActivities;
+		TOptional<TArray<FConcertSyncReplicationActivity>> ReturnResult_EnumerateActivities;
 		
 		//~ Begin IReplicationWorkspace Interface
-		virtual TOptional<int64> ProduceClientLeaveReplicationActivity(const FGuid& EndpointId, const FConcertSyncReplicationPayload_LeaveReplication& EventData) override
+		virtual TOptional<int64> ProduceReplicationActivity(const FGuid& EndpointId, const FConcertSyncReplicationEvent& EventData) override
 		{
-			LastCall_ProduceClientLeaveReplicationActivity = MakeTuple(EndpointId, EventData);
-			return ReturnResult_ProduceClientLeaveReplicationActivity;
-		}
-		virtual TOptional<int64> ProduceClientMuteReplicationActivity(const FGuid& EndpointId, const FConcertSyncReplicationPayload_Mute& EventData) override
-		{
-			LastCall_ProduceClientMuteReplicationActivity = MakeTuple(EndpointId, EventData);
-			return ReturnResult_ProduceClientMuteReplicationActivity;
+			switch (EventData.ActivityType)
+			{
+			case EConcertSyncReplicationActivityType::LeaveReplication:
+				{
+					FConcertSyncReplicationPayload_LeaveReplication Data;
+					const bool bSuccess = EventData.GetPayload(Data);
+					ensure(bSuccess);
+					
+					LastCall_ProduceClientLeaveReplicationActivity = MakeTuple(EndpointId, Data);
+					return ReturnResult_ProduceClientLeaveReplicationActivity;
+				}
+			case EConcertSyncReplicationActivityType::Mute:
+				{
+					FConcertSyncReplicationPayload_Mute Data;
+					const bool bSuccess = EventData.GetPayload(Data);
+					ensure(bSuccess);
+					
+					LastCall_ProduceClientMuteReplicationActivity = MakeTuple(EndpointId, Data);
+					return ReturnResult_ProduceClientMuteReplicationActivity;
+				}
+				
+			case EConcertSyncReplicationActivityType::None: [[fallthrough]];
+			case EConcertSyncReplicationActivityType::Count: [[fallthrough]];
+			default: ensure(false); return {};
+			}
 		}
 
-		virtual bool GetLastLeaveReplicationActivityByClient(const FConcertSessionClientInfo& InClientInfo, FConcertSyncReplicationPayload_LeaveReplication& OutLeaveReplication) const override
+		virtual bool GetLastReplicationActivityByClient(const FConcertSessionClientInfo& InClientInfo, FConcertSyncReplicationActivity& OutActivity) const override
 		{
-			LastCall_GetLastLeaveReplicationActivityByClient = MakeTuple(InClientInfo);
-			if (ReturnResult_GetLastLeaveReplicationActivityByClient)
+			LastCall_GetLastReplicationActivityByClient = MakeTuple(InClientInfo);
+			if (ReturnResult_GetLastReplicationActivityByClient)
 			{
-				OutLeaveReplication = *ReturnResult_GetLastLeaveReplicationActivityByClient;
+				OutActivity = *ReturnResult_GetLastReplicationActivityByClient;
 			}
-			return ReturnResult_GetLastLeaveReplicationActivityByClient.IsSet();
+			return ReturnResult_GetLastReplicationActivityByClient.IsSet();
 		}
-		virtual bool GetLeaveReplicationActivityById(const int64 ActivityId, FConcertSyncReplicationPayload_LeaveReplication& OutLeaveReplication) const override
+
+		virtual bool GetReplicationEventById(const int64 ActivityId, FConcertSyncReplicationEvent& OutEvent) const override
 		{
 			LastCall_GetLeaveReplicationActivityById = MakeTuple(ActivityId);
-			if (ReturnResult_GetLeaveReplicationActivityById)
+			if (ReturnResult_GetReplicationEventById)
 			{
-				OutLeaveReplication = *ReturnResult_GetLeaveReplicationActivityById;
+				OutEvent = *ReturnResult_GetReplicationEventById;
 			}
-			return ReturnResult_GetLeaveReplicationActivityById.IsSet();
+			return ReturnResult_GetReplicationEventById.IsSet();
 		}
-		virtual void EnumerateMuteActivities(TFunctionRef<EBreakBehavior(const FConcertSyncReplicationActivity& Activity)> Callback) const override
+		
+		virtual void EnumerateReplicationActivities(TFunctionRef<EBreakBehavior(const FConcertSyncReplicationActivity& Activity)> Callback) const override
 		{
-			if (!ReturnResult_EnumerateMuteActivities)
+			if (!ReturnResult_EnumerateActivities)
 			{
 				return;
 			}
 			
-			for (const FConcertSyncReplicationActivity& Activity : *ReturnResult_EnumerateMuteActivities)
+			for (const FConcertSyncReplicationActivity& Activity : *ReturnResult_EnumerateActivities)
 			{
 				if (Callback(Activity) == EBreakBehavior::Break)
 				{
diff --git a/Engine/Plugins/Developer/Concert/ConcertSync/ConcertSyncTest/Source/ConcertSyncTest/Private/Replication/Util/Mocks/ReplicationWorkspaceEmptyMock.h b/Engine/Plugins/Developer/Concert/ConcertSync/ConcertSyncTest/Source/ConcertSyncTest/Private/Replication/Util/Mocks/ReplicationWorkspaceEmptyMock.h
--- a/Engine/Plugins/Developer/Concert/ConcertSync/ConcertSyncTest/Source/ConcertSyncTest/Private/Replication/Util/Mocks/ReplicationWorkspaceEmptyMock.h
+++ b/Engine/Plugins/Developer/Concert/ConcertSync/ConcertSyncTest/Source/ConcertSyncTest/Private/Replication/Util/Mocks/ReplicationWorkspaceEmptyMock.h
@@ -12,11 +12,10 @@ namespace UE::ConcertSyncTests::Replication
 	public:
 		
 		//~ Begin IReplicationWorkspace Interface
-		virtual TOptional<int64> ProduceClientLeaveReplicationActivity(const FGuid&, const FConcertSyncReplicationPayload_LeaveReplication&) override { return {}; }
-		virtual TOptional<int64> ProduceClientMuteReplicationActivity(const FGuid& EndpointId, const FConcertSyncReplicationPayload_Mute& EventData) override { return {}; }
-		virtual bool GetLastLeaveReplicationActivityByClient(const FConcertSessionClientInfo& InClientInfo, FConcertSyncReplicationPayload_LeaveReplication& OutLeaveReplication) const override { return false; }
-		virtual bool GetLeaveReplicationActivityById(const int64 ActivityId, FConcertSyncReplicationPayload_LeaveReplication& OutLeaveReplication) const override { return false; }
-		virtual void EnumerateMuteActivities(TFunctionRef<EBreakBehavior(const FConcertSyncReplicationActivity& Activity)> Callback) const override {}
+		virtual TOptional<int64> ProduceReplicationActivity(const FGuid& EndpointId, const FConcertSyncReplicationEvent& EventData) override { return {}; }
+		virtual bool GetLastReplicationActivityByClient(const FConcertSessionClientInfo& InClientInfo, FConcertSyncReplicationActivity& OutActivity) const override { return false; }
+		virtual bool GetReplicationEventById(const int64 ActivityId, FConcertSyncReplicationEvent& OutEvent) const override { return false; }
+		virtual void EnumerateReplicationActivities(TFunctionRef<EBreakBehavior(const FConcertSyncReplicationActivity& Activity)> Callback) const override {}
 		//~ End IReplicationWorkspace Interface
 	};
 }
