Fix failing ensure on server when client requests frequency info just after remote client has left

diff --git a/Engine/Plugins/Developer/Concert/ConcertSync/ConcertSyncServer/Source/ConcertSyncServer/Private/Replication/ConcertServerReplicationManager.cpp b/Engine/Plugins/Developer/Concert/ConcertSync/ConcertSyncServer/Source/ConcertSyncServer/Private/Replication/ConcertServerReplicationManager.cpp
--- a/Engine/Plugins/Developer/Concert/ConcertSync/ConcertSyncServer/Source/ConcertSyncServer/Private/Replication/ConcertServerReplicationManager.cpp
+++ b/Engine/Plugins/Developer/Concert/ConcertSync/ConcertSyncServer/Source/ConcertSyncServer/Private/Replication/ConcertServerReplicationManager.cpp
@@ -322,7 +322,7 @@ namespace UE::ConcertSyncServer::Replication
 	FConcertObjectReplicationSettings FConcertServerReplicationManager::GetObjectFrequencySettings(const FConcertReplicatedObjectId& Object) const
 	{
 		const TUniquePtr<FConcertReplicationClient>* Client = Clients.Find(Object.SenderEndpointId);
-		if (!ensureMsgf(Client, TEXT("Caller is trying to retrieve non-existing client")))
+		if (!Client)
 		{
 			UE_LOG(LogConcert, Warning, TEXT("Requested frequency settings for unknown client %s"), *Object.SenderEndpointId.ToString());
 			return {};
@@ -332,7 +332,7 @@ namespace UE::ConcertSyncServer::Replication
 			{
 				return Description.BaseDescription.Identifier == Object.StreamId;
 			});
-		if (!ensureMsgf(Stream, TEXT("Caller is trying to retrieve an object that is not registered with the client")))
+		if (!Stream)
 		{
 			UE_LOG(LogConcert, Warning, TEXT("Requested frequency settings for unknown stream %s and object %s"), *Object.StreamId.ToString(), *Object.Object.ToString());
 			return {};
