Remove inappropriate use of UE_LIFETIMEBOUND from void returning functions

diff --git a/Engine/Plugins/Animation/TweeningUtils/Source/TweeningUtilsEditor/Public/Math/Models/FunctionBasedTweenModel.h b/Engine/Plugins/Animation/TweeningUtils/Source/TweeningUtilsEditor/Public/Math/Models/FunctionBasedTweenModel.h
--- a/Engine/Plugins/Animation/TweeningUtils/Source/TweeningUtilsEditor/Public/Math/Models/FunctionBasedTweenModel.h
+++ b/Engine/Plugins/Animation/TweeningUtils/Source/TweeningUtilsEditor/Public/Math/Models/FunctionBasedTweenModel.h
@@ -19,7 +19,7 @@ public:
 	explicit FFunctionBasedTweenModel(const ITweenFunction& InInitialFunction UE_LIFETIMEBOUND) : TweenFunction(&InInitialFunction) {}
 
 	/** Sets the tween function. You are responsible to make sure it outlives this model. */
-	virtual void SetTweenFunction(const ITweenFunction& InTweenFunction UE_LIFETIMEBOUND) { TweenFunction = &InTweenFunction; }
+	virtual void SetTweenFunction(const ITweenFunction& InTweenFunction) { TweenFunction = &InTweenFunction; }
 	const ITweenFunction& GetTweenFunction() const { return *TweenFunction; }
 	
 private:
diff --git a/Engine/Plugins/Animation/TweeningUtils/Source/TweeningUtilsEditor/Public/Math/Models/UserPreferredTweenProxy.h b/Engine/Plugins/Animation/TweeningUtils/Source/TweeningUtilsEditor/Public/Math/Models/UserPreferredTweenProxy.h
--- a/Engine/Plugins/Animation/TweeningUtils/Source/TweeningUtilsEditor/Public/Math/Models/UserPreferredTweenProxy.h
+++ b/Engine/Plugins/Animation/TweeningUtils/Source/TweeningUtilsEditor/Public/Math/Models/UserPreferredTweenProxy.h
@@ -33,7 +33,7 @@ public:
 	}
 
 	//~ Begin FTweenModel Interface
-	virtual void SetTweenFunction(const ITweenFunction& InTweenFunction UE_LIFETIMEBOUND)
+	virtual void SetTweenFunction(const ITweenFunction& InTweenFunction)
 	{
 		const bool bIsContained = OwningFunctionContainer.IsValidIndex(OwningFunctionContainer.IndexOf(InTweenFunction));
 		FString FunctionId = OwningFunctionContainer.GetFunctionIdentifier(InTweenFunction);
diff --git a/Engine/Plugins/Developer/Concert/ConcertApp/MultiUserClient/Source/MultiUserClient/Private/Widgets/ActiveSession/Replication/Client/Multi/Columns/SMuteToggle.h b/Engine/Plugins/Developer/Concert/ConcertApp/MultiUserClient/Source/MultiUserClient/Private/Widgets/ActiveSession/Replication/Client/Multi/Columns/SMuteToggle.h
--- a/Engine/Plugins/Developer/Concert/ConcertApp/MultiUserClient/Source/MultiUserClient/Private/Widgets/ActiveSession/Replication/Client/Multi/Columns/SMuteToggle.h
+++ b/Engine/Plugins/Developer/Concert/ConcertApp/MultiUserClient/Source/MultiUserClient/Private/Widgets/ActiveSession/Replication/Client/Multi/Columns/SMuteToggle.h
@@ -23,7 +23,12 @@ namespace UE::MultiUserClient::Replication
 		SLATE_BEGIN_ARGS(SMuteToggle) {}
 		SLATE_END_ARGS()
 
-		void Construct(const FArguments& InArgs, FSoftObjectPath InObjectPath, FMuteChangeTracker& InMuteChangeTracker UE_LIFETIMEBOUND);
+		/**
+		 * @param InArgs Widget arguments
+		 * @param InObjectPath The path of the object for which to toggle the mute state
+		 * @param InMuteChangeTracker The caller ensures it outlives the lifetime of the widget.
+		 */
+		void Construct(const FArguments& InArgs, FSoftObjectPath InObjectPath, FMuteChangeTracker& InMuteChangeTracker);
 
 	private:
 
diff --git a/Engine/Plugins/Developer/Concert/ConcertApp/MultiUserClient/Source/MultiUserClient/Private/Widgets/ActiveSession/Replication/Client/Multi/SAllClientsView.h b/Engine/Plugins/Developer/Concert/ConcertApp/MultiUserClient/Source/MultiUserClient/Private/Widgets/ActiveSession/Replication/Client/Multi/SAllClientsView.h
--- a/Engine/Plugins/Developer/Concert/ConcertApp/MultiUserClient/Source/MultiUserClient/Private/Widgets/ActiveSession/Replication/Client/Multi/SAllClientsView.h
+++ b/Engine/Plugins/Developer/Concert/ConcertApp/MultiUserClient/Source/MultiUserClient/Private/Widgets/ActiveSession/Replication/Client/Multi/SAllClientsView.h
@@ -24,7 +24,11 @@ namespace UE::MultiUserClient::Replication
 		SLATE_BEGIN_ARGS(SAllClientsView){}
 		SLATE_END_ARGS()
 
-		void Construct(const FArguments&, TSharedRef<IConcertClient> InConcertClient, FMultiUserReplicationManager& InMultiUserReplicationManager UE_LIFETIMEBOUND);
+		/**
+		 * @param InConcertClient The local client instance
+		 * @param InMultiUserReplicationManager The caller ensures it outlives the lifetime of the widget.
+		 */
+		void Construct(const FArguments&, TSharedRef<IConcertClient> InConcertClient, FMultiUserReplicationManager& InMultiUserReplicationManager);
 
 		const TSharedPtr<SMultiClientView>& GetClientView() const { return ClientView; }
 		
diff --git a/Engine/Plugins/Developer/Concert/ConcertApp/MultiUserClient/Source/MultiUserClient/Private/Widgets/ActiveSession/Replication/Client/Multi/ViewOptions/SMultiViewOptions.h b/Engine/Plugins/Developer/Concert/ConcertApp/MultiUserClient/Source/MultiUserClient/Private/Widgets/ActiveSession/Replication/Client/Multi/ViewOptions/SMultiViewOptions.h
--- a/Engine/Plugins/Developer/Concert/ConcertApp/MultiUserClient/Source/MultiUserClient/Private/Widgets/ActiveSession/Replication/Client/Multi/ViewOptions/SMultiViewOptions.h
+++ b/Engine/Plugins/Developer/Concert/ConcertApp/MultiUserClient/Source/MultiUserClient/Private/Widgets/ActiveSession/Replication/Client/Multi/ViewOptions/SMultiViewOptions.h
@@ -18,7 +18,11 @@ namespace UE::MultiUserClient::Replication
 		SLATE_BEGIN_ARGS(SMultiViewOptions){}
 		SLATE_END_ARGS()
 
-		void Construct(const FArguments& InArgs, FMultiViewOptions& InViewOptions UE_LIFETIMEBOUND);
+		/**
+		 * @param InArgs The widget args
+		 * @param InViewOptions Options governing the display settings. The caller ensures it outlives the lifetime of the widget.
+		 */
+		void Construct(const FArguments& InArgs, FMultiViewOptions& InViewOptions);
 
 	private:
 
diff --git a/Engine/Plugins/Developer/Concert/ConcertApp/MultiUserClient/Source/MultiUserClient/Private/Widgets/ActiveSession/Replication/Client/PropertySelection/SPropertySelectionComboButton.h b/Engine/Plugins/Developer/Concert/ConcertApp/MultiUserClient/Source/MultiUserClient/Private/Widgets/ActiveSession/Replication/Client/PropertySelection/SPropertySelectionComboButton.h
--- a/Engine/Plugins/Developer/Concert/ConcertApp/MultiUserClient/Source/MultiUserClient/Private/Widgets/ActiveSession/Replication/Client/PropertySelection/SPropertySelectionComboButton.h
+++ b/Engine/Plugins/Developer/Concert/ConcertApp/MultiUserClient/Source/MultiUserClient/Private/Widgets/ActiveSession/Replication/Client/PropertySelection/SPropertySelectionComboButton.h
@@ -32,7 +32,11 @@ namespace UE::MultiUserClient::Replication
 			SLATE_EVENT(FGetObjectDisplayString, GetObjectDisplayString)
 		SLATE_END_ARGS()
 
-		void Construct(const FArguments& InArgs, FUserPropertySelector& InPropertySelector UE_LIFETIMEBOUND);
+		/**
+		 * @param InArgs Widget arguments
+		 * @param InPropertySelector User to change the selected properties. The caller ensures it outlives the lifetime of the widget.
+		 */
+		void Construct(const FArguments& InArgs, FUserPropertySelector& InPropertySelector);
 
 		/** Refreshes the properties that the user can select given the objects currently displayed in the bottom view. */
 		void RefreshSelectableProperties(TConstArrayView<ConcertSharedSlate::FObjectGroup> DisplayedObjectGroups) const;
diff --git a/Engine/Plugins/Developer/Concert/ConcertApp/MultiUserClient/Source/MultiUserClient/Private/Widgets/ActiveSession/Replication/Client/SPresetComboButton.h b/Engine/Plugins/Developer/Concert/ConcertApp/MultiUserClient/Source/MultiUserClient/Private/Widgets/ActiveSession/Replication/Client/SPresetComboButton.h
--- a/Engine/Plugins/Developer/Concert/ConcertApp/MultiUserClient/Source/MultiUserClient/Private/Widgets/ActiveSession/Replication/Client/SPresetComboButton.h
+++ b/Engine/Plugins/Developer/Concert/ConcertApp/MultiUserClient/Source/MultiUserClient/Private/Widgets/ActiveSession/Replication/Client/SPresetComboButton.h
@@ -31,7 +31,12 @@ namespace UE::MultiUserClient::Replication
 			
 		SLATE_END_ARGS()
 
-		void Construct(const FArguments& InArgs, const IConcertClient& InClient UE_LIFETIMEBOUND, FPresetManager& InPresetManager UE_LIFETIMEBOUND);
+		/**
+		 * @param InArgs Widget arguments
+		 * @param InClient Caller ensures that it outlives the lifetime of this widget.
+		 * @param InPresetManager Caller ensures that it outlives the lifetime of this widget.
+		 */
+		void Construct(const FArguments& InArgs, const IConcertClient& InClient, FPresetManager& InPresetManager);
 
 	private:
 
diff --git a/Engine/Plugins/Developer/Concert/ConcertInsights/ConcertInsightsClient/Source/ConcertInsightsClient/Private/StatusBar/StatusBarExtension.h b/Engine/Plugins/Developer/Concert/ConcertInsights/ConcertInsightsClient/Source/ConcertInsightsClient/Private/StatusBar/StatusBarExtension.h
--- a/Engine/Plugins/Developer/Concert/ConcertInsights/ConcertInsightsClient/Source/ConcertInsightsClient/Private/StatusBar/StatusBarExtension.h
+++ b/Engine/Plugins/Developer/Concert/ConcertInsights/ConcertInsightsClient/Source/ConcertInsightsClient/Private/StatusBar/StatusBarExtension.h
@@ -9,6 +9,6 @@ namespace UE::ConcertInsightsClient
 	class FClientTraceControls;
 	
 	/** Adds a sub-menu to the Multi User status bar. */
-	void ExtendMultiUserStatusBarWithInsights(FClientTraceControls& Controls UE_LIFETIMEBOUND);
+	void ExtendMultiUserStatusBarWithInsights(FClientTraceControls& Controls);
 }
 
diff --git a/Engine/Plugins/Developer/Concert/ConcertSync/ConcertSyncCore/Source/ConcertSyncCore/Public/Replication/Processing/ReplicationDataQueuer.h b/Engine/Plugins/Developer/Concert/ConcertSync/ConcertSyncCore/Source/ConcertSyncCore/Public/Replication/Processing/ReplicationDataQueuer.h
--- a/Engine/Plugins/Developer/Concert/ConcertSync/ConcertSyncCore/Source/ConcertSyncCore/Public/Replication/Processing/ReplicationDataQueuer.h
+++ b/Engine/Plugins/Developer/Concert/ConcertSync/ConcertSyncCore/Source/ConcertSyncCore/Public/Replication/Processing/ReplicationDataQueuer.h
@@ -37,8 +37,11 @@ namespace UE::ConcertSyncCore
 
 	protected:
 
-		/** Called by subclass factory functions. */
-		void BindToCache(FObjectReplicationCache& InReplicationCache UE_LIFETIMEBOUND);
+		/**
+		 * Called by subclass factory functions.
+		 * @param InReplicationCache The cache that the data is sent to. The caller ensures it outlives the lifetime of the widget.
+		 */
+		void BindToCache(FObjectReplicationCache& InReplicationCache);
 		
 	private:
 
diff --git a/Engine/Plugins/Developer/Concert/ConcertSync/ConcertSyncServer/Source/ConcertSyncServer/Private/ConcertSyncServer.h b/Engine/Plugins/Developer/Concert/ConcertSync/ConcertSyncServer/Source/ConcertSyncServer/Private/ConcertSyncServer.h
--- a/Engine/Plugins/Developer/Concert/ConcertSync/ConcertSyncServer/Source/ConcertSyncServer/Private/ConcertSyncServer.h
+++ b/Engine/Plugins/Developer/Concert/ConcertSync/ConcertSyncServer/Source/ConcertSyncServer/Private/ConcertSyncServer.h
@@ -63,8 +63,17 @@ private:
 
 	void CreateSequencerManager(const TSharedRef<FConcertSyncServerLiveSession>& InLiveSession);
 	void DestroySequencerManager(const TSharedRef<FConcertSyncServerLiveSession>& InLiveSession);
-	
-	void CreateReplicationManager(const TSharedRef<IConcertServerSession>& InSession, UE::ConcertSyncServer::Replication::IReplicationWorkspace& InWorkspace UE_LIFETIMEBOUND, EConcertSyncSessionFlags InSessionFlags);
+
+	/**
+	 * Initializes the replication manager for the given session.
+	 * @param InSession The session to bind the manager to
+	 * @param InWorkspace The database interface that manager interacts with. The caller ensures it outlives the lifetime of the widget.
+	 * @param InSessionFlags Session flags, which determine optional features.
+	 */
+	void CreateReplicationManager(
+		const TSharedRef<IConcertServerSession>& InSession, UE::ConcertSyncServer::Replication::IReplicationWorkspace& InWorkspace,
+		EConcertSyncSessionFlags InSessionFlags
+		);
 	void DestroyReplicationManager(const FGuid& SessionId); 
 
 	bool CreateLiveSession(const TSharedRef<IConcertServerSession>& InSession, const FInternalLiveSessionCreationParams& AdditionalParams);
