Multi User Server: Fix Concert.Replication.LogAuthorityRequestsAndResponsesOnServer not logging authority requests & responses:
- Move stream & authority CVar logging to new ReplicationCVars.h.

diff --git a/Engine/Plugins/Developer/Concert/ConcertSync/ConcertSyncServer/Source/ConcertSyncServer/Private/Replication/AuthorityManager.cpp b/Engine/Plugins/Developer/Concert/ConcertSync/ConcertSyncServer/Source/ConcertSyncServer/Private/Replication/AuthorityManager.cpp
--- a/Engine/Plugins/Developer/Concert/ConcertSync/ConcertSyncServer/Source/ConcertSyncServer/Private/Replication/AuthorityManager.cpp
+++ b/Engine/Plugins/Developer/Concert/ConcertSync/ConcertSyncServer/Source/ConcertSyncServer/Private/Replication/AuthorityManager.cpp
@@ -10,6 +10,8 @@
 #include "Replication/Messages/ChangeAuthority.h"
 #include "Replication/Misc/IReplicationGroundTruth.h"
 #include "Util/GroundTruthOverride.h"
+#include "Util/LogUtils.h"
+#include "Util/ReplicationCVars.h"
 
 namespace UE::ConcertSyncServer::Replication
 {
@@ -271,7 +273,11 @@ namespace UE::ConcertSyncServer::Replication
 		// This log does two things: 1. Identify issues in unit tests / at runtime 2. Warn about possibly malicious attempts when the server runs.
 		UE_CLOG(Request.TakeAuthority.IsEmpty() && Request.ReleaseAuthority.IsEmpty(), LogConcert, Warning, TEXT("Received invalid authority request (TakeAuthority.Num() == 0 and ReleaseAuthority.Num() == 0)"));
 		
+		const FGuid ClientId = Context.SourceEndpointId;
+		LogNetworkMessage(CVarLogAuthorityRequestsAndResponsesOnServer, Request, [&](){ return GetClientName(*Session, ClientId); });
 		InternalApplyChangeAuthorityRequest(Context.SourceEndpointId, Request, Response.RejectedObjects, Response.SyncControl);
+		LogNetworkMessage(CVarLogAuthorityRequestsAndResponsesOnServer, Response, [&](){ return GetClientName(*Session, ClientId); });
+		
 		Response.ErrorCode = EReplicationResponseErrorCode::Handled;
 		return EConcertSessionResponseCode::Success;
 	}
diff --git a/Engine/Plugins/Developer/Concert/ConcertSync/ConcertSyncServer/Source/ConcertSyncServer/Private/Replication/ConcertServerReplicationManager.cpp b/Engine/Plugins/Developer/Concert/ConcertSync/ConcertSyncServer/Source/ConcertSyncServer/Private/Replication/ConcertServerReplicationManager.cpp
--- a/Engine/Plugins/Developer/Concert/ConcertSync/ConcertSyncServer/Source/ConcertSyncServer/Private/Replication/ConcertServerReplicationManager.cpp
+++ b/Engine/Plugins/Developer/Concert/ConcertSync/ConcertSyncServer/Source/ConcertSyncServer/Private/Replication/ConcertServerReplicationManager.cpp
@@ -10,24 +10,19 @@
 #include "Replication/Data/ClientQueriedInfo.h"
 #include "Replication/Formats/FullObjectFormat.h"
 #include "Replication/IReplicationWorkspace.h"
-#include "Replication/Messages/PutState.h"
 #include "Replication/Messages/ChangeStream.h"
 #include "Replication/Messages/ClientQuery.h"
 #include "Replication/Messages/Handshake.h"
+#include "Replication/Messages/PutState.h"
 #include "Replication/Messages/ReplicationActivity.h"
 #include "Replication/Messages/RestoreContent.h"
 #include "Replication/Processing/ObjectReplicationCache.h"
 #include "Util/JoinRequestValidation.h"
 #include "Util/LogUtils.h"
+#include "Util/ReplicationCVars.h"
 
 namespace UE::ConcertSyncServer::Replication
 {
-	TAutoConsoleVariable<bool> CVarLogAuthorityRequestsAndResponsesOnServer(
-		TEXT("Concert.Replication.LogAuthorityRequestsAndResponsesOnServer"),
-		false,
-		TEXT("Whether to log changes to authority.")
-		);
-	
 	FConcertServerReplicationManager::FConcertServerReplicationManager(
 		TSharedRef<IConcertServerSession> InLiveSession,
 		IReplicationWorkspace& InServerWorkspace,
diff --git a/Engine/Plugins/Developer/Concert/ConcertSync/ConcertSyncServer/Source/ConcertSyncServer/Private/Replication/ConcertServerReplicationManager_ChangeStream.cpp b/Engine/Plugins/Developer/Concert/ConcertSync/ConcertSyncServer/Source/ConcertSyncServer/Private/Replication/ConcertServerReplicationManager_ChangeStream.cpp
--- a/Engine/Plugins/Developer/Concert/ConcertSync/ConcertSyncServer/Source/ConcertSyncServer/Private/Replication/ConcertServerReplicationManager_ChangeStream.cpp
+++ b/Engine/Plugins/Developer/Concert/ConcertSync/ConcertSyncServer/Source/ConcertSyncServer/Private/Replication/ConcertServerReplicationManager_ChangeStream.cpp
@@ -9,16 +9,11 @@
 #include "Replication/ConcertReplicationClient.h"
 #include "Replication/Messages/ChangeStream.h"
 #include "Util/LogUtils.h"
+#include "Util/ReplicationCVars.h"
 #include "Util/StreamChangeValidation.h"
 
 namespace UE::ConcertSyncServer::Replication
 {
-	TAutoConsoleVariable<bool> CVarLogStreamRequestsAndResponsesOnServer(
-		TEXT("Concert.Replication.LogStreamRequestsAndResponsesOnServer"),
-		false,
-		TEXT("Whether to log changes to streams.")
-		);
-
 	EConcertSessionResponseCode FConcertServerReplicationManager::HandleChangeStreamRequest(
 		const FConcertSessionContext& ConcertSessionContext,
 		const FConcertReplication_ChangeStream_Request& Request,
diff --git a/Engine/Plugins/Developer/Concert/ConcertSync/ConcertSyncServer/Source/ConcertSyncServer/Private/Replication/Util/ReplicationCVars.cpp b/Engine/Plugins/Developer/Concert/ConcertSync/ConcertSyncServer/Source/ConcertSyncServer/Private/Replication/Util/ReplicationCVars.cpp
new file mode 100644
--- /dev/null
+++ b/Engine/Plugins/Developer/Concert/ConcertSync/ConcertSyncServer/Source/ConcertSyncServer/Private/Replication/Util/ReplicationCVars.cpp
@@ -0,0 +1,18 @@
+﻿// Copyright Epic Games, Inc. All Rights Reserved.
+
+#include "ReplicationCVars.h"
+
+namespace UE::ConcertSyncServer::Replication
+{
+TAutoConsoleVariable<bool> CVarLogStreamRequestsAndResponsesOnServer(
+	TEXT("Concert.Replication.LogStreamRequestsAndResponsesOnServer"),
+	false,
+	TEXT("Whether to log changes to streams.")
+	);
+	
+TAutoConsoleVariable<bool> CVarLogAuthorityRequestsAndResponsesOnServer(
+	TEXT("Concert.Replication.LogAuthorityRequestsAndResponsesOnServer"),
+	false,
+	TEXT("Whether to log changes to authority.")
+	);
+}
\ No newline at end of file
diff --git a/Engine/Plugins/Developer/Concert/ConcertSync/ConcertSyncServer/Source/ConcertSyncServer/Private/Replication/Util/ReplicationCVars.h b/Engine/Plugins/Developer/Concert/ConcertSync/ConcertSyncServer/Source/ConcertSyncServer/Private/Replication/Util/ReplicationCVars.h
new file mode 100644
--- /dev/null
+++ b/Engine/Plugins/Developer/Concert/ConcertSync/ConcertSyncServer/Source/ConcertSyncServer/Private/Replication/Util/ReplicationCVars.h
@@ -0,0 +1,14 @@
+﻿// Copyright Epic Games, Inc. All Rights Reserved.
+
+#pragma once
+
+#include "HAL/IConsoleManager.h"
+
+namespace UE::ConcertSyncServer::Replication
+{
+/** Whether to log requests and responses to streams on the server. */
+extern TAutoConsoleVariable<bool> CVarLogStreamRequestsAndResponsesOnServer;
+	
+/** Whether to log requests and responses to authority on the server. */
+extern TAutoConsoleVariable<bool> CVarLogAuthorityRequestsAndResponsesOnServer;
+}
