Fix missing #include.

This appears to be a long-standing issue masked by unity build.

#Synced-CL 52670723

diff --git a/Engine/Plugins/Developer/Concert/ConcertSync/ConcertSyncServer/Source/ConcertSyncServer/Private/Replication/Util/LogUtils.h b/Engine/Plugins/Developer/Concert/ConcertSync/ConcertSyncServer/Source/ConcertSyncServer/Private/Replication/Util/LogUtils.h
--- a/Engine/Plugins/Developer/Concert/ConcertSync/ConcertSyncServer/Source/ConcertSyncServer/Private/Replication/Util/LogUtils.h
+++ b/Engine/Plugins/Developer/Concert/ConcertSync/ConcertSyncServer/Source/ConcertSyncServer/Private/Replication/Util/LogUtils.h
@@ -6,6 +6,7 @@
 #include "JsonObjectConverter.h"
 #include "HAL/IConsoleManager.h"
 #include "Logging/LogMacros.h"
+#include "IConcertSession.h"
 
 #include <type_traits>
 
